Cycle the database closed/open in the session test to validate that it can be opened again

#rnx

diff --git a/Engine/Plugins/Developer/Concert/ConcertSync/ConcertSyncTest/Source/ConcertSyncTest/Private/ConcertDatabaseTests.cpp b/Engine/Plugins/Developer/Concert/ConcertSync/ConcertSyncTest/Source/ConcertSyncTest/Private/ConcertDatabaseTests.cpp
--- a/Engine/Plugins/Developer/Concert/ConcertSync/ConcertSyncTest/Source/ConcertSyncTest/Private/ConcertDatabaseTests.cpp
+++ b/Engine/Plugins/Developer/Concert/ConcertSync/ConcertSyncTest/Source/ConcertSyncTest/Private/ConcertDatabaseTests.cpp
@@ -84,6 +84,19 @@ bool FConcertSessionDatabaseTest::RunTest(const FString& Parameters)
 				AddError(FString::Printf(TEXT("Failed to add package activity to server database: %s"), *SessionDatabase_Server.GetLastError()));
 			}
 		}
+
+		// Close and re-open the client database to verify that it isn't corrupted
+		if (SessionDatabase_Client.IsValid())
+		{
+			if (!SessionDatabase_Client.Close())
+			{
+				AddError(TEXT("Failed to close client database"));
+			}
+			if (!SessionDatabase_Client.Open(TestSessionPath_Client))
+			{
+				AddError(TEXT("Failed to open client database"));
+			}
+		}
 	}
 
 	// Populate the client database from the server data
