Initialize test variables to 0.

diff --git a/Engine/Plugins/Developer/Concert/ConcertSync/ConcertSyncTest/Source/ConcertSyncTest/Private/ConcertDataStoreTests.h b/Engine/Plugins/Developer/Concert/ConcertSync/ConcertSyncTest/Source/ConcertSyncTest/Private/ConcertDataStoreTests.h
--- a/Engine/Plugins/Developer/Concert/ConcertSync/ConcertSyncTest/Source/ConcertSyncTest/Private/ConcertDataStoreTests.h
+++ b/Engine/Plugins/Developer/Concert/ConcertSync/ConcertSyncTest/Source/ConcertSyncTest/Private/ConcertDataStoreTests.h
@@ -12,13 +12,13 @@ struct FConcertDataStore_CustomTypeTest
 	GENERATED_BODY()
 
 	UPROPERTY(VisibleAnywhere, Category = "Concert Message")
-	int8 Int8Value;
+	int8 Int8Value = 0;
 
 	UPROPERTY(VisibleAnywhere, Category = "Concert Message")
-	int64 Int64Value;
+	int64 Int64Value = 0;
 
 	UPROPERTY(VisibleAnywhere, Category = "Concert Message")
-	float FloatValue;
+	float FloatValue = 0;
 
 	UPROPERTY(VisibleAnywhere, Category = "Concert Message")
 	TArray<int> IntArray;
