Detect dependend transactions

diff --git a/Engine/Plugins/Developer/Concert/ConcertSync/ConcertSyncCore/Source/ConcertSyncCore/Private/HistoryEdition/HistoryAnalysis.cpp b/Engine/Plugins/Developer/Concert/ConcertSync/ConcertSyncCore/Source/ConcertSyncCore/Private/HistoryEdition/HistoryAnalysis.cpp
new file mode 100644
--- /dev/null
+++ b/Engine/Plugins/Developer/Concert/ConcertSync/ConcertSyncCore/Source/ConcertSyncCore/Private/HistoryEdition/HistoryAnalysis.cpp
@@ -0,0 +1,113 @@
+﻿// Copyright Epic Games, Inc. All Rights Reserved.
+
+#include "HistoryEdition/HistoryAnalysis.h"
+
+#include "HistoryEdition/DependencyGraphBuilder.h"
+
+namespace UE::ConcertSyncCore
+{
+	FHistoryDeletionRequirements AnalyseActivityDeletion(const TSet<FActivityID>& ActivitiesToDelete, const FConcertSyncSessionDatabase& Database, bool bAddActivitiesToDelete)
+	{
+		const FActivityDependencyGraph Graph = BuildDependencyGraphFrom(Database);
+		return AnalyseActivityDeletion(ActivitiesToDelete, Graph, bAddActivitiesToDelete);
+	}
+	
+	FHistoryDeletionRequirements AnalyseActivityDeletion(const TSet<FActivityID>& ActivitiesToDelete, const FActivityDependencyGraph& DependencyGraph, bool bAddActivitiesToDelete)
+	{
+		FHistoryDeletionRequirements Result;
+
+		TSet<FActivityNodeID> HardDoubleEnqueuingProtection;
+		TQueue<FActivityNodeID> HardDependencyActivitiesToAnalyse;
+		for (const FActivityID ActivityToDelete : ActivitiesToDelete)
+		{
+			const TOptional<FActivityNodeID> NodeID = DependencyGraph.FindNodeByActivity(ActivityToDelete);
+			if (ensureMsgf(NodeID, TEXT("Graph does not correspond to ActivitiesToDelete")))
+			{
+				HardDoubleEnqueuingProtection.Add(*NodeID);
+				HardDependencyActivitiesToAnalyse.Enqueue(*NodeID);
+			}
+		}
+
+		TSet<FActivityNodeID> PossibleDoubleEnqueuingProtection;
+		TQueue<FActivityNodeID> PossibleDependencyActivitiesToAnalyse;
+		FActivityNodeID CurrentActivityID;
+
+		/* We check the hard dependencies first.
+		 * Why? Example:
+		 *
+		 *		R
+		 *	   / \
+		 *	  A   B
+		 *	   \ /
+		 *	    C
+		 *
+		 * The edges C -> A -> R are possible dependencies.
+		 * The edges C -> B -> R are hard dependencies.
+		 *
+		 * Now: delete R.
+		 * We want C to be marked has a hard dependency.
+		 */
+		while (HardDependencyActivitiesToAnalyse.Dequeue(CurrentActivityID))
+		{
+			const FActivityNode& ActivityNode = DependencyGraph.GetNodeById(CurrentActivityID);
+
+			const FActivityID ActivityID = ActivityNode.GetActivityId();
+			if (bAddActivitiesToDelete || !ActivitiesToDelete.Contains(ActivityID))
+			{
+				Result.HardDependencies.Add(ActivityID);
+			}
+			
+			for (const FActivityNodeID& ChildID : ActivityNode.GetAffectedChildren())
+			{
+				const FActivityNode& ChildNode = DependencyGraph.GetNodeById(ChildID);
+
+				// Performance: The below iterates the edge list twice but usually there will 1 or 2 entries
+				if (ChildNode.DependsOnActivity(ActivityID, DependencyGraph, {}, EDependencyStrength::HardDependency))
+				{
+					if (!HardDoubleEnqueuingProtection.Contains(ChildID))
+					{
+						HardDoubleEnqueuingProtection.Add(ChildID);
+						HardDependencyActivitiesToAnalyse.Enqueue(ChildID);
+					}
+				}
+				else if (ChildNode.DependsOnActivity(ActivityID, DependencyGraph, {}, EDependencyStrength::PossibleDependency))
+				{
+					if (!HardDoubleEnqueuingProtection.Contains(ChildID) && !PossibleDoubleEnqueuingProtection.Contains(ChildID))
+					{
+						PossibleDoubleEnqueuingProtection.Add(ChildID);
+						PossibleDependencyActivitiesToAnalyse.Enqueue(ChildID);
+					}
+				}
+			}
+		}
+
+		// Any possible dependencies that are not also hard dependencies can be added now
+		while (PossibleDependencyActivitiesToAnalyse.Dequeue(CurrentActivityID))
+		{
+			// This would imply a hard dependency - hard dependency takes precedence over possible dependency
+			if (HardDoubleEnqueuingProtection.Contains(CurrentActivityID))
+			{
+				continue;
+			}
+			
+			const FActivityNode& ActivityNode = DependencyGraph.GetNodeById(CurrentActivityID);
+			const FActivityID ActivityID = ActivityNode.GetActivityId();
+			Result.PossibleDependencies.Add(ActivityNode.GetActivityId());
+			for (const FActivityNodeID& ChildID : ActivityNode.GetAffectedChildren())
+			{
+				const FActivityNode& ChildNode = DependencyGraph.GetNodeById(ChildID);
+				if (ChildNode.DependsOnActivity(ActivityID, DependencyGraph, {}, EDependencyStrength::PossibleDependency)
+					&& !HardDoubleEnqueuingProtection.Contains(ChildID)
+					&& !PossibleDoubleEnqueuingProtection.Contains(ChildID))
+				{
+					PossibleDoubleEnqueuingProtection.Add(ChildID);
+					PossibleDependencyActivitiesToAnalyse.Enqueue(ChildID);
+				}
+			}
+		}
+
+		return Result;
+	}
+}
+
+
diff --git a/Engine/Plugins/Developer/Concert/ConcertSync/ConcertSyncCore/Source/ConcertSyncCore/Public/HistoryEdition/HistoryAnalysis.h b/Engine/Plugins/Developer/Concert/ConcertSync/ConcertSyncCore/Source/ConcertSyncCore/Public/HistoryEdition/HistoryAnalysis.h
new file mode 100644
--- /dev/null
+++ b/Engine/Plugins/Developer/Concert/ConcertSync/ConcertSyncCore/Source/ConcertSyncCore/Public/HistoryEdition/HistoryAnalysis.h
@@ -0,0 +1,37 @@
+﻿// Copyright Epic Games, Inc. All Rights Reserved.
+
+#pragma once
+
+#include "CoreMinimal.h"
+#include "ActivityDependencyGraph.h"
+#include "ConcertMessageData.h"
+
+class FConcertSyncSessionDatabase;
+
+namespace UE::ConcertSyncCore
+{
+	/** Describes the activities that must be considered when deleting an activity */
+	struct CONCERTSYNCCORE_API FHistoryDeletionRequirements
+	{
+		/** The activities that must be removed. */
+		TSet<FActivityID> HardDependencies;
+		/**
+		 * The activities may want to be removed. It's not certain that they are affected (but it should be safe to keep them in).
+		 * This will not contain any elements in HardDependencies.
+		 */
+		TSet<FActivityID> PossibleDependencies;
+	};
+
+	/** Utility function for one-off operations: just computes the dependency graph before calling AnalyseActivityDeletion. */
+	CONCERTSYNCCORE_API FHistoryDeletionRequirements AnalyseActivityDeletion(const TSet<FActivityID>& ActivitiesToDelete, const FConcertSyncSessionDatabase& Database, bool bAddActivitiesToDelete = false);
+
+	/**
+	 * Given a set of activities to be deleted, returns which activities 1. must be and 2. may want to be deleted in addition.
+	 * @param ActivitiesToDelete The activities that should be removed
+	 * @param DependencyGraph The graph encoding the activity dependencies
+	 * @param bAddActivitiesToDelete Whether to add ActivitiesToDelete to the result's HardDependencies
+	 * @return The activities to remove as well if ActivitiesToDelete is removed; ActivitiesToDelete is included in HardDependencies if bAddActivitiesToDelete == true.
+	 */
+	CONCERTSYNCCORE_API FHistoryDeletionRequirements AnalyseActivityDeletion(const TSet<FActivityID>& ActivitiesToDelete, const FActivityDependencyGraph& DependencyGraph, bool bAddActivitiesToDelete = false);
+}
+
diff --git a/Engine/Plugins/Developer/Concert/ConcertSync/ConcertSyncTest/Source/ConcertSyncTest/Private/HistoryEdition/HistoryAnalysisTests.cpp b/Engine/Plugins/Developer/Concert/ConcertSync/ConcertSyncTest/Source/ConcertSyncTest/Private/HistoryEdition/HistoryAnalysisTests.cpp
new file mode 100644
--- /dev/null
+++ b/Engine/Plugins/Developer/Concert/ConcertSync/ConcertSyncTest/Source/ConcertSyncTest/Private/HistoryEdition/HistoryAnalysisTests.cpp
@@ -0,0 +1,157 @@
+// Copyright Epic Games, Inc. All Rights Reserved.
+
+#include "ConcertLogGlobal.h"
+
+#include "ConcertSyncSessionDatabase.h"
+#include "HistoryEdition/HistoryAnalysis.h"
+#include "RenameEditAndDeleteMapsFlow.h"
+#include "ScopedSessionDatabase.h"
+
+#include "Algo/AllOf.h"
+
+#include "Core/Tests/Containers/TestUtils.h"
+
+#include "HistoryEdition/DebugDependencyGraph.h"
+#include "HistoryEdition/DependencyGraphBuilder.h"
+
+#include "Misc/AutomationTest.h"
+
+namespace UE::ConcertSyncTests
+{
+	bool ValidateRequirements(
+		const FString& TestBaseName,
+		FAutomationTestBase& Test,
+		const RenameEditAndDeleteMapsFlowTest::TTestActivityArray<FActivityID> Activities,
+		const ConcertSyncCore::FHistoryDeletionRequirements& ToValidate = {},
+		const TSet<RenameEditAndDeleteMapsFlowTest::ETestActivity>& ExpectedHardDependencies = {},
+		const TSet<RenameEditAndDeleteMapsFlowTest::ETestActivity>& ExpectedPossibleDependencies = {}
+		)
+	{
+		using namespace RenameEditAndDeleteMapsFlowTest;
+
+		const bool bHardDependenciesAreCorrect = Algo::AllOf(ExpectedHardDependencies, [&TestBaseName, &Test, &Activities, &ToValidate](ETestActivity Activity)
+		{
+			const bool bContained = ToValidate.HardDependencies.Contains(Activities[Activity]);
+			Test.TestTrue(FString::Printf(TEXT("%s: %s is a hard dependency"), *TestBaseName, *LexToString(Activity)), bContained);
+			return bContained;
+		});
+		Test.TestTrue(FString::Printf(TEXT("%s: Hard dependencies are correct"), *TestBaseName), bHardDependenciesAreCorrect);
+		
+		const bool bPossibleDependenciesAreCorrect = Algo::AllOf(ExpectedPossibleDependencies, [&TestBaseName, &Test, &Activities, &ToValidate](ETestActivity Activity)
+		{
+			const bool bContained = ToValidate.PossibleDependencies.Contains(Activities[Activity]);
+			Test.TestTrue(FString::Printf(TEXT("%s: %s is a possible dependency"), *TestBaseName, *LexToString(Activity)), bContained);
+			return bContained;
+		});
+		Test.TestTrue(FString::Printf(TEXT("%s: Possible dependencies are correct"), *TestBaseName), bPossibleDependenciesAreCorrect);
+		
+		const TSet<ETestActivity> ExpectedExcludedActivities = AllActivities().Difference(ExpectedHardDependencies.Union(ExpectedPossibleDependencies));
+		const bool bAllOtherActivitiesExcluded = Algo::AllOf(ExpectedExcludedActivities, [&TestBaseName, &Test, &Activities, &ToValidate](ETestActivity Activity)
+		{
+			const bool bNotContained = !ToValidate.HardDependencies.Contains(Activities[Activity]) && !ToValidate.PossibleDependencies.Contains(Activities[Activity]);
+			Test.TestTrue(FString::Printf(TEXT("%s: %s is no dependency"), *TestBaseName, *LexToString(Activity)), bNotContained);
+			return bNotContained;
+		});
+		Test.TestTrue(FString::Printf(TEXT("%s: No unexpected dependencies"), *TestBaseName), bAllOtherActivitiesExcluded);
+
+		return bHardDependenciesAreCorrect && bPossibleDependenciesAreCorrect && bAllOtherActivitiesExcluded;
+	}
+}
+
+namespace UE::ConcertSyncTests::AnalysisTests
+{
+	IMPLEMENT_SIMPLE_AUTOMATION_TEST(FAnalyseDeletionDependencies, "Concert.History.Analysis.AnalyseDeletionDependencies", EAutomationTestFlags::ApplicationContextMask | EAutomationTestFlags::SmokeFilter)
+	bool FAnalyseDeletionDependencies::RunTest(const FString& Parameters)
+	{
+		using namespace RenameEditAndDeleteMapsFlowTest;
+		using namespace ConcertSyncCore;
+		
+		FScopedSessionDatabase SessionDatabase(*this);
+		const TTestActivityArray<FActivityID> Activities = CreateActivityHistory(SessionDatabase, SessionDatabase.GetEndpoint());
+		const FActivityDependencyGraph DependencyGraph = BuildDependencyGraphFrom(SessionDatabase);
+		UE_LOG(LogConcert, Log, TEXT("%s tested graph in Graphviz format:\n\n%s"), *GetTestFullName(), *ConcertSyncCore::Graphviz::ExportToGraphviz(DependencyGraph, SessionDatabase));
+
+		// Delete /Game/Foo > Nearly everything has hard dependency
+		{
+			const FHistoryDeletionRequirements DeleteFooRequirements = AnalyseActivityDeletion({ Activities[_1_NewPackageFoo] }, DependencyGraph);
+			/* _1_NewPackageFoo: is what we're "deleting".
+			 * _5_SavePackageBar: Bar is created as result of a rename but has no dependency to _1_NewPackageFoo.
+			 * All other activities transitively depend on _1_NewPackageFoo (put above log into GraphViz to visualise).
+			 * 
+			 * Note: The transaction activities (_3_EditActor, _4_EditActor) have possible dependencies BUT they do have hard dependencies to _2_AddActor. This is why they must be in HardDependencies, too.
+			 */
+			const TSet<ETestActivity> ExcludedActivities{ _1_NewPackageFoo, _5_SavePackageBar };
+			const TSet<ETestActivity> HardDependencies = AllActivities().Difference(ExcludedActivities);
+			const bool bDeleteAllCorrect = ValidateRequirements(TEXT("Delete /Game/Foo"), *this, Activities, DeleteFooRequirements, HardDependencies);
+			TestTrue(TEXT("Delete /Game/Foo is correct"), bDeleteAllCorrect);
+		}
+
+		// Delete rename transaction > No dependencies
+		{
+			const FHistoryDeletionRequirements DeleteRenameRequirements = AnalyseActivityDeletion({ Activities[_3_RenameActor] }, DependencyGraph);
+			TestEqual(TEXT("Delete renaming actor: HardDependencies.Num() == 0"), DeleteRenameRequirements.HardDependencies.Num(), 0);
+			TestEqual(TEXT("Delete renaming actor: PossibleDependencies.Num() == 1"), DeleteRenameRequirements.PossibleDependencies.Num(), 1);
+			TestTrue(TEXT("Delete renaming actor: Edit activity may depend on deleted activity"), DeleteRenameRequirements.PossibleDependencies.Contains(Activities[_4_EditActor]));
+		}
+
+		// Delete actor creation > All transactions operating on actor are hard dependencies
+		{
+			const FHistoryDeletionRequirements DeleteCreateActorRequirements = AnalyseActivityDeletion({ Activities[_2_AddActor] }, DependencyGraph);
+			TestEqual(TEXT("Delete actor creation: HardDependencies.Num() == 2"), DeleteCreateActorRequirements.HardDependencies.Num(), 2);
+			TestEqual(TEXT("Delete actor creation: PossibleDependencies.Num() == 0"), DeleteCreateActorRequirements.PossibleDependencies.Num(), 0);
+			TestTrue(TEXT("Delete actor creation: Rename depends on created actor"), DeleteCreateActorRequirements.HardDependencies.Contains(Activities[_3_RenameActor]));
+			TestTrue(TEXT("Delete actor creation: Edit depends on created actor"), DeleteCreateActorRequirements.HardDependencies.Contains(Activities[_4_EditActor]));
+		}
+
+		return true;
+	}
+
+	/**
+	 * Suppose:
+	 *
+	 *		R
+	 *	   / \
+	 *	  A   B
+	 *	   \ /
+	 *	    L
+	 *
+	 * The edges L -> A -> R are possible dependencies.
+	 * The edges L -> B -> R are hard dependencies.
+	 *
+	 * The test: delete R.
+	 * We want L to be marked has a hard dependency.
+	 */
+	IMPLEMENT_SIMPLE_AUTOMATION_TEST(FPossibleDependencyOrderedBeforeHardDependency, "Concert.History.Analysis.PossibleDependencyOrderedBeforeHardDependency", EAutomationTestFlags::ApplicationContextMask | EAutomationTestFlags::SmokeFilter)
+	bool FPossibleDependencyOrderedBeforeHardDependency::RunTest(const FString& Parameters)
+	{
+		using namespace ConcertSyncCore;
+
+		constexpr FActivityID RootActivityID = 1;
+		constexpr FActivityID AActivityID = 2;
+		constexpr FActivityID BActivityID = 3;
+		constexpr FActivityID LeafActivityID  = 4;
+		
+		FActivityDependencyGraph DependencyGraph;
+		const FActivityNodeID RootNodeID = DependencyGraph.AddActivity(RootActivityID);
+		const FActivityNodeID ANodeID = DependencyGraph.AddActivity(AActivityID);
+		const FActivityNodeID BNodeID = DependencyGraph.AddActivity(BActivityID);
+		const FActivityNodeID LeafNodeID = DependencyGraph.AddActivity(LeafActivityID);
+
+		// Add the weak dependency first so the algorithm finds it first when iterating
+		DependencyGraph.AddDependency(LeafNodeID, FActivityDependencyEdge(ANodeID, EActivityDependencyReason::EditAfterPreviousPackageEdit, EDependencyStrength::PossibleDependency));
+		DependencyGraph.AddDependency(ANodeID, FActivityDependencyEdge(RootNodeID, EActivityDependencyReason::EditAfterPreviousPackageEdit, EDependencyStrength::PossibleDependency));
+		DependencyGraph.AddDependency(LeafNodeID, FActivityDependencyEdge(BNodeID, EActivityDependencyReason::EditAfterPreviousPackageEdit, EDependencyStrength::HardDependency));
+		DependencyGraph.AddDependency(BNodeID, FActivityDependencyEdge(RootNodeID, EActivityDependencyReason::EditAfterPreviousPackageEdit, EDependencyStrength::HardDependency));
+
+		const FHistoryDeletionRequirements DeleteFooRequirements = AnalyseActivityDeletion({ RootActivityID }, DependencyGraph);
+
+		TestEqual(TEXT("HardDependencies.Num() == 2"), DeleteFooRequirements.HardDependencies.Num(), 2);
+		TestTrue(TEXT("HardDependencies.Contains(B)"), DeleteFooRequirements.HardDependencies.Contains(BActivityID));
+		TestTrue(TEXT("HardDependencies.Contains(L)"), DeleteFooRequirements.HardDependencies.Contains(LeafActivityID));
+		
+		TestEqual(TEXT("PossibleDependencies.Num() == 1"), DeleteFooRequirements.PossibleDependencies.Num(), 1);
+		TestTrue(TEXT("PossibleDependencies.Contains(A)"), DeleteFooRequirements.PossibleDependencies.Contains(AActivityID));
+
+		return true;
+	}
+}
\ No newline at end of file
