Concert some unit tests from smoke to engine filter tests

diff --git a/Engine/Plugins/Developer/Concert/ConcertSync/ConcertSyncTest/Source/ConcertSyncTest/Private/HistoryEdition/DependencyGraphBuilderTests.cpp b/Engine/Plugins/Developer/Concert/ConcertSync/ConcertSyncTest/Source/ConcertSyncTest/Private/HistoryEdition/DependencyGraphBuilderTests.cpp
--- a/Engine/Plugins/Developer/Concert/ConcertSync/ConcertSyncTest/Source/ConcertSyncTest/Private/HistoryEdition/DependencyGraphBuilderTests.cpp
+++ b/Engine/Plugins/Developer/Concert/ConcertSync/ConcertSyncTest/Source/ConcertSyncTest/Private/HistoryEdition/DependencyGraphBuilderTests.cpp
@@ -39,7 +39,7 @@ namespace UE::ConcertSyncTests::RenameEditAndDeleteMapsFlowTest
 	 *	7 -> 5 (PackageRename)
 	 *	8 -> 5 (PackageRemoval)
 	 */
-	IMPLEMENT_SIMPLE_AUTOMATION_TEST(FRenameEditAndDeleteMapsFlowTest, "Concert.History.BuildGraph.RenameEditAndDeleteMapsFlow", EAutomationTestFlags::ApplicationContextMask | EAutomationTestFlags::SmokeFilter)
+	IMPLEMENT_SIMPLE_AUTOMATION_TEST(FRenameEditAndDeleteMapsFlowTest, "Concert.History.BuildGraph.RenameEditAndDeleteMapsFlow", EAutomationTestFlags::ApplicationContextMask | EAutomationTestFlags::EngineFilter)
 	bool FRenameEditAndDeleteMapsFlowTest::RunTest(const FString& Parameters)
 	{
 		FScopedSessionDatabase SessionDatabase(*this);
@@ -172,7 +172,7 @@ namespace UE::ConcertSyncTests::DeletingAndRecreatingActorIsHardDependency
 	 *
 	 * 3 > 2 is a hard dependency (removing 2 would result in attempting to create the actor twice, which is invalid).
 	 */
-	IMPLEMENT_SIMPLE_AUTOMATION_TEST(FDeletingAndRecreatingActorIsHardDependency, "Concert.History.BuildGraph.DeletingAndRecreatingActorIsHardDependency", EAutomationTestFlags::ApplicationContextMask | EAutomationTestFlags::SmokeFilter)
+	IMPLEMENT_SIMPLE_AUTOMATION_TEST(FDeletingAndRecreatingActorIsHardDependency, "Concert.History.BuildGraph.DeletingAndRecreatingActorIsHardDependency", EAutomationTestFlags::ApplicationContextMask | EAutomationTestFlags::EngineFilter)
 	bool FDeletingAndRecreatingActorIsHardDependency::RunTest(const FString& Parameters)
 	{
 		using namespace ConcertSyncCore;
@@ -289,7 +289,7 @@ namespace UE::ConcertSyncTests::IndirectPackageDependencyTest
 	 *
 	 * This is relevant because the actor's construction script may depend query data from the data asset.
 	 */
-	IMPLEMENT_SIMPLE_AUTOMATION_TEST(FIndirectPackageDependencyTest, "Concert.History.BuildGraph.IndirectPackageDependencyTest", EAutomationTestFlags::ApplicationContextMask | EAutomationTestFlags::SmokeFilter)
+	IMPLEMENT_SIMPLE_AUTOMATION_TEST(FIndirectPackageDependencyTest, "Concert.History.BuildGraph.IndirectPackageDependencyTest", EAutomationTestFlags::ApplicationContextMask | EAutomationTestFlags::EngineFilter)
 	bool FIndirectPackageDependencyTest::RunTest(const FString& Parameters)
 	{
 		// TODO:
diff --git a/Engine/Plugins/Developer/Concert/ConcertSync/ConcertSyncTest/Source/ConcertSyncTest/Private/HistoryEdition/HistoryAnalysisTests.cpp b/Engine/Plugins/Developer/Concert/ConcertSync/ConcertSyncTest/Source/ConcertSyncTest/Private/HistoryEdition/HistoryAnalysisTests.cpp
--- a/Engine/Plugins/Developer/Concert/ConcertSync/ConcertSyncTest/Source/ConcertSyncTest/Private/HistoryEdition/HistoryAnalysisTests.cpp
+++ b/Engine/Plugins/Developer/Concert/ConcertSync/ConcertSyncTest/Source/ConcertSyncTest/Private/HistoryEdition/HistoryAnalysisTests.cpp
@@ -60,7 +60,7 @@ namespace UE::ConcertSyncTests
 
 namespace UE::ConcertSyncTests::AnalysisTests
 {
-	IMPLEMENT_SIMPLE_AUTOMATION_TEST(FAnalyseDeletionDependencies, "Concert.History.Analysis.AnalyseDeletionDependencies", EAutomationTestFlags::ApplicationContextMask | EAutomationTestFlags::SmokeFilter)
+	IMPLEMENT_SIMPLE_AUTOMATION_TEST(FAnalyseDeletionDependencies, "Concert.History.Analysis.AnalyseDeletionDependencies", EAutomationTestFlags::ApplicationContextMask | EAutomationTestFlags::EngineFilter)
 	bool FAnalyseDeletionDependencies::RunTest(const FString& Parameters)
 	{
 		using namespace RenameEditAndDeleteMapsFlowTest;
@@ -121,7 +121,7 @@ namespace UE::ConcertSyncTests::AnalysisTests
 	 * The test: delete R.
 	 * We want L to be marked has a hard dependency.
 	 */
-	IMPLEMENT_SIMPLE_AUTOMATION_TEST(FPossibleDependencyOrderedBeforeHardDependency, "Concert.History.Analysis.PossibleDependencyOrderedBeforeHardDependency", EAutomationTestFlags::ApplicationContextMask | EAutomationTestFlags::SmokeFilter)
+	IMPLEMENT_SIMPLE_AUTOMATION_TEST(FPossibleDependencyOrderedBeforeHardDependency, "Concert.History.Analysis.PossibleDependencyOrderedBeforeHardDependency", EAutomationTestFlags::ApplicationContextMask | EAutomationTestFlags::EngineFilter)
 	bool FPossibleDependencyOrderedBeforeHardDependency::RunTest(const FString& Parameters)
 	{
 		using namespace ConcertSyncCore;
