Save Package activities related to a Rename activity should be a possible dependency

Unshelved from pending changelist '20009909':

diff --git a/Engine/Plugins/Developer/Concert/ConcertApp/MultiUserServer/Source/MultiUserServer/Private/MultiUserServerConsoleVariables.cpp b/Engine/Plugins/Developer/Concert/ConcertApp/MultiUserServer/Source/MultiUserServer/Private/MultiUserServerConsoleVariables.cpp
new file mode 100644
--- /dev/null
+++ b/Engine/Plugins/Developer/Concert/ConcertApp/MultiUserServer/Source/MultiUserServer/Private/MultiUserServerConsoleVariables.cpp
@@ -0,0 +1,7 @@
+﻿// Copyright Epic Games, Inc. All Rights Reserved.
+
+#include "MultiUserServerConsoleVariables.h"
+
+#include "Containers/UnrealString.h"
+
+TAutoConsoleVariable<bool> UE::MultiUserServer::ConsoleVariables::CVarLogActivityDependencyGraphOnDelete(TEXT("MultiUserServer.LogActivityDependencyGraphOnDelete"), false, TEXT("When a user requests to delete any activity, log the dependency graph (in GraphViz format)"));
diff --git a/Engine/Plugins/Developer/Concert/ConcertApp/MultiUserServer/Source/MultiUserServer/Private/MultiUserServerConsoleVariables.h b/Engine/Plugins/Developer/Concert/ConcertApp/MultiUserServer/Source/MultiUserServer/Private/MultiUserServerConsoleVariables.h
new file mode 100644
--- /dev/null
+++ b/Engine/Plugins/Developer/Concert/ConcertApp/MultiUserServer/Source/MultiUserServer/Private/MultiUserServerConsoleVariables.h
@@ -0,0 +1,12 @@
+﻿// Copyright Epic Games, Inc. All Rights Reserved.
+
+#pragma once
+
+#include "CoreMinimal.h"
+#include "HAL/IConsoleManager.h"
+
+namespace UE::MultiUserServer::ConsoleVariables
+{
+	/** When a user requests to delete any activity, log the dependency graph (in GraphViz format) */
+	extern TAutoConsoleVariable<bool> CVarLogActivityDependencyGraphOnDelete;
+}
\ No newline at end of file
diff --git a/Engine/Plugins/Developer/Concert/ConcertApp/MultiUserServer/Source/MultiUserServer/Private/Widgets/SessionTabs/Archived/ArchivedConcertSessionTab.cpp b/Engine/Plugins/Developer/Concert/ConcertApp/MultiUserServer/Source/MultiUserServer/Private/Widgets/SessionTabs/Archived/ArchivedConcertSessionTab.cpp
--- a/Engine/Plugins/Developer/Concert/ConcertApp/MultiUserServer/Source/MultiUserServer/Private/Widgets/SessionTabs/Archived/ArchivedConcertSessionTab.cpp
+++ b/Engine/Plugins/Developer/Concert/ConcertApp/MultiUserServer/Source/MultiUserServer/Private/Widgets/SessionTabs/Archived/ArchivedConcertSessionTab.cpp
@@ -4,6 +4,7 @@
 
 #include "ArchivedSessionHistoryController.h"
 #include "ConcertLogGlobal.h"
+#include "MultiUserServerConsoleVariables.h"
 #include "MultiUserServerModule.h"
 #include "Window/ModalWindowManager.h"
 
@@ -13,6 +14,7 @@
 #include "Dialog/SMessageDialog.h"
 
 #include "HistoryEdition/ActivityNode.h"
+#include "HistoryEdition/DebugDependencyGraph.h"
 #include "HistoryEdition/DependencyGraphBuilder.h"
 #include "HistoryEdition/HistoryAnalysis.h"
 #include "HistoryEdition/HistoryDeletion.h"
@@ -62,12 +64,17 @@ void FArchivedConcertSessionTab::OnRequestDeleteActivity(const TSet<TSharedRef<F
 	
 	if (const TOptional<FConcertSyncSessionDatabaseNonNullPtr> SessionDatabase = SyncServer->GetArchivedSessionDatabase(InspectedSessionID))
 	{
+		const FActivityDependencyGraph DependencyGraph = BuildDependencyGraphFrom(*SessionDatabase);
+		if (UE::MultiUserServer::ConsoleVariables::CVarLogActivityDependencyGraphOnDelete.GetValueOnGameThread())
+		{
+			UE_LOG(LogConcert, Log, TEXT("%s"), *UE::ConcertSyncCore::Graphviz::ExportToGraphviz(DependencyGraph, *SessionDatabase));
+		}
+		
 		TSet<FActivityID> RequestedForDelete;
 		Algo::Transform(ActivitiesToDelete, RequestedForDelete, [](const TSharedRef<FConcertSessionActivity>& Activity)
 		{
 			return Activity->Activity.ActivityId;
 		});
-		const FActivityDependencyGraph DependencyGraph = BuildDependencyGraphFrom(*SessionDatabase);
 		FHistoryDeletionRequirements DeletionRequirements = AnalyseActivityDeletion(RequestedForDelete, DependencyGraph, true);
 
 		TWeakPtr<const FArchivedConcertSessionTab> WeakTabThis = SharedThis(this);
diff --git a/Engine/Plugins/Developer/Concert/ConcertSync/ConcertSyncCore/Source/ConcertSyncCore/Private/HistoryEdition/ActivityDependencyGraph.cpp b/Engine/Plugins/Developer/Concert/ConcertSync/ConcertSyncCore/Source/ConcertSyncCore/Private/HistoryEdition/ActivityDependencyGraph.cpp
--- a/Engine/Plugins/Developer/Concert/ConcertSync/ConcertSyncCore/Source/ConcertSyncCore/Private/HistoryEdition/ActivityDependencyGraph.cpp
+++ b/Engine/Plugins/Developer/Concert/ConcertSync/ConcertSyncCore/Source/ConcertSyncCore/Private/HistoryEdition/ActivityDependencyGraph.cpp
@@ -8,13 +8,13 @@
 
 #include "Algo/ForEach.h"
 
-UE::ConcertSyncCore::FActivityNodeID UE::ConcertSyncCore::FActivityDependencyGraph::AddActivity(int64 ActivityIndex)
+UE::ConcertSyncCore::FActivityNodeID UE::ConcertSyncCore::FActivityDependencyGraph::AddActivity(int64 ActivityIndex, EActivityNodeFlags NodeFlags)
 {
 	const TOptional<FActivityNodeID> PreexistingID = FindNodeByActivity(ActivityIndex);
 	if (ensureMsgf(!PreexistingID.IsSet(), TEXT("Activity is already registered!")))
 	{
 		FActivityNodeID NodeID(Nodes.Num());
-		Nodes.Emplace(ActivityIndex, NodeID);
+		Nodes.Emplace(ActivityIndex, NodeID, NodeFlags);
 		return NodeID;
 	}
 
diff --git a/Engine/Plugins/Developer/Concert/ConcertSync/ConcertSyncCore/Source/ConcertSyncCore/Private/HistoryEdition/DependencyGraphBuilder.cpp b/Engine/Plugins/Developer/Concert/ConcertSync/ConcertSyncCore/Source/ConcertSyncCore/Private/HistoryEdition/DependencyGraphBuilder.cpp
--- a/Engine/Plugins/Developer/Concert/ConcertSync/ConcertSyncCore/Source/ConcertSyncCore/Private/HistoryEdition/DependencyGraphBuilder.cpp
+++ b/Engine/Plugins/Developer/Concert/ConcertSync/ConcertSyncCore/Source/ConcertSyncCore/Private/HistoryEdition/DependencyGraphBuilder.cpp
@@ -185,7 +185,12 @@ namespace UE::ConcertSyncCore
 
 	void FActivityDependencyGraphBuildAlgorithm::DiscoverPackageDependencies(const FConcertSyncActivity& Activity, const FConcertSyncPackageEventMetaData& EventData)
 	{
-		const FActivityNodeID NodeID = Graph.AddActivity(Activity.ActivityId);
+		const FActivityNodeID NodeID = Graph.AddActivity(
+			Activity.ActivityId,
+			EventData.PackageInfo.PackageUpdateType == EConcertPackageUpdateType::Renamed
+				? EActivityNodeFlags::RenameActivity
+				: EActivityNodeFlags::None
+			);
 		switch (EventData.PackageInfo.PackageUpdateType)
 		{
 		case EConcertPackageUpdateType::Added:
diff --git a/Engine/Plugins/Developer/Concert/ConcertSync/ConcertSyncCore/Source/ConcertSyncCore/Private/HistoryEdition/HistoryAnalysis.cpp b/Engine/Plugins/Developer/Concert/ConcertSync/ConcertSyncCore/Source/ConcertSyncCore/Private/HistoryEdition/HistoryAnalysis.cpp
--- a/Engine/Plugins/Developer/Concert/ConcertSync/ConcertSyncCore/Source/ConcertSyncCore/Private/HistoryEdition/HistoryAnalysis.cpp
+++ b/Engine/Plugins/Developer/Concert/ConcertSync/ConcertSyncCore/Source/ConcertSyncCore/Private/HistoryEdition/HistoryAnalysis.cpp
@@ -2,6 +2,9 @@
 
 #include "HistoryEdition/HistoryAnalysis.h"
 
+#include "Algo/AnyOf.h"
+#include "Algo/Find.h"
+
 #include "HistoryEdition/ActivityGraphIDs.h"
 #include "HistoryEdition/DependencyGraphBuilder.h"
 
@@ -14,11 +17,21 @@ namespace UE::ConcertSyncCore
 		const FActivityDependencyGraph Graph = BuildDependencyGraphFrom(Database);
 		return AnalyseActivityDeletion(ActivitiesToDelete, Graph, bAddActivitiesToDelete);
 	}
+
+	/**
+	 * If PossibleRenameActivityNode is a rename activity, then this function adds the activity that created the renamed-to package.
+	 * This implies that PossibleRenameActivityNode is EConcertSyncActivityEventType::Package and has EConcertPackageUpdateType::Renamed.
+	 */
+	static void AddSavePackageActivityAssociatedWithRenamePackageActivity(const FActivityDependencyGraph& DependencyGraph, const FActivityNode& PossibleRenameActivityNode, TSet<FActivityNodeID>& DoubleEnqueuingProtection, TQueue<FActivityNodeID>& ActivitiesToAnalyse);
 	
 	FHistoryDeletionRequirements AnalyseActivityDeletion(const TSet<FActivityID>& ActivitiesToDelete, const FActivityDependencyGraph& DependencyGraph, bool bAddActivitiesToDelete)
 	{
 		FHistoryDeletionRequirements Result;
 
+		TSet<FActivityNodeID> PossibleDoubleEnqueuingProtection;
+		TQueue<FActivityNodeID> PossibleDependencyActivitiesToAnalyse;
+		FActivityNodeID CurrentActivityID;
+		
 		TSet<FActivityNodeID> HardDoubleEnqueuingProtection;
 		TQueue<FActivityNodeID> HardDependencyActivitiesToAnalyse;
 		for (const FActivityID ActivityToDelete : ActivitiesToDelete)
@@ -28,13 +41,10 @@ namespace UE::ConcertSyncCore
 			{
 				HardDoubleEnqueuingProtection.Add(*NodeID);
 				HardDependencyActivitiesToAnalyse.Enqueue(*NodeID);
+				AddSavePackageActivityAssociatedWithRenamePackageActivity(DependencyGraph, DependencyGraph.GetNodeById(*NodeID), PossibleDoubleEnqueuingProtection, PossibleDependencyActivitiesToAnalyse);
 			}
 		}
-
-		TSet<FActivityNodeID> PossibleDoubleEnqueuingProtection;
-		TQueue<FActivityNodeID> PossibleDependencyActivitiesToAnalyse;
-		FActivityNodeID CurrentActivityID;
-
+		
 		/* We check the hard dependencies first.
 		 * Why? Example:
 		 *
@@ -71,6 +81,7 @@ namespace UE::ConcertSyncCore
 					{
 						HardDoubleEnqueuingProtection.Add(ChildID);
 						HardDependencyActivitiesToAnalyse.Enqueue(ChildID);
+						AddSavePackageActivityAssociatedWithRenamePackageActivity(DependencyGraph, ChildNode, PossibleDoubleEnqueuingProtection, PossibleDependencyActivitiesToAnalyse);
 					}
 				}
 				else if (ChildNode.DependsOnActivity(ActivityID, DependencyGraph, {}, EDependencyStrength::PossibleDependency))
@@ -79,6 +90,7 @@ namespace UE::ConcertSyncCore
 					{
 						PossibleDoubleEnqueuingProtection.Add(ChildID);
 						PossibleDependencyActivitiesToAnalyse.Enqueue(ChildID);
+						AddSavePackageActivityAssociatedWithRenamePackageActivity(DependencyGraph, ChildNode, PossibleDoubleEnqueuingProtection, PossibleDependencyActivitiesToAnalyse);
 					}
 				}
 			}
@@ -105,12 +117,37 @@ namespace UE::ConcertSyncCore
 				{
 					PossibleDoubleEnqueuingProtection.Add(ChildID);
 					PossibleDependencyActivitiesToAnalyse.Enqueue(ChildID);
+					AddSavePackageActivityAssociatedWithRenamePackageActivity(DependencyGraph, ChildNode, PossibleDoubleEnqueuingProtection, PossibleDependencyActivitiesToAnalyse);
 				}
 			}
 		}
 
 		return Result;
 	}
+
+	static void AddSavePackageActivityAssociatedWithRenamePackageActivity(const FActivityDependencyGraph& DependencyGraph, const FActivityNode& PossibleRenameActivityNode, TSet<FActivityNodeID>& DoubleEnqueuingProtection, TQueue<FActivityNodeID>& ActivitiesToAnalyse)
+	{
+		if ((PossibleRenameActivityNode.GetNodeFlags() & EActivityNodeFlags::RenameActivity) == EActivityNodeFlags::None)
+		{
+			return;
+		}
+
+		// Generally rename activities depends on at least one activity and at most on two:
+			// 1. There is always a PossibleDependency with reason PackageCreation: this is the activity that created the renamed-to package; that's the one we want to add here)
+			// 2. There is sometimes a HardDependency with reason PackageCreation in case the renamed package was created within the same session
+		const FActivityDependencyEdge* PackageCreationDependency = Algo::FindByPredicate(PossibleRenameActivityNode.GetDependencies(),
+			[](const FActivityDependencyEdge& Dependency)
+			{
+				const bool bIsAssociatedCreationActivity = Dependency.GetDependencyReason() == EActivityDependencyReason::PackageCreation && Dependency.GetDependencyStrength() == EDependencyStrength::PossibleDependency;
+				return bIsAssociatedCreationActivity;
+			});
+		if (ensureMsgf(PackageCreationDependency, TEXT("Rename activities should have a possible dependency to the activity that creatd the renamed-to package"))
+			&& !DoubleEnqueuingProtection.Contains(PackageCreationDependency->GetDependedOnNodeID()))
+		{
+			DoubleEnqueuingProtection.Add(PackageCreationDependency->GetDependedOnNodeID());
+			ActivitiesToAnalyse.Enqueue(PackageCreationDependency->GetDependedOnNodeID());
+		}
+	}
 }
 
 
diff --git a/Engine/Plugins/Developer/Concert/ConcertSync/ConcertSyncCore/Source/ConcertSyncCore/Public/HistoryEdition/ActivityDependencyGraph.h b/Engine/Plugins/Developer/Concert/ConcertSync/ConcertSyncCore/Source/ConcertSyncCore/Public/HistoryEdition/ActivityDependencyGraph.h
--- a/Engine/Plugins/Developer/Concert/ConcertSync/ConcertSyncCore/Source/ConcertSyncCore/Public/HistoryEdition/ActivityDependencyGraph.h
+++ b/Engine/Plugins/Developer/Concert/ConcertSync/ConcertSyncCore/Source/ConcertSyncCore/Public/HistoryEdition/ActivityDependencyGraph.h
@@ -47,10 +47,12 @@ namespace UE::ConcertSyncCore
 	 *  3. Create actor in level
 	 *  4. Edit actor to reference data assets A
 	 *  5. Edit actor whose construction script reads data from A
-	 * Resulting Graph:     
-	 *     1   3
-	 *    / \ / 
-	 *   2   4
+	 * Resulting Graph:
+	 *     1
+	 *	   I
+	 *     2   3
+	 *      \ / 
+	 *       4
 	 *       I
 	 *       5
 	 */
@@ -59,7 +61,7 @@ namespace UE::ConcertSyncCore
 	public:
 		
 		/** Adds a node. Invalidates the results of GetNodeById. */
-		FActivityNodeID AddActivity(int64 ActivityIndex);
+		FActivityNodeID AddActivity(int64 ActivityIndex, EActivityNodeFlags NodeFlags = EActivityNodeFlags::None);
 		
 		/** Adds a dependency from an existing node to another existing node */
 		bool AddDependency(FActivityNodeID From, FActivityDependencyEdge To);
diff --git a/Engine/Plugins/Developer/Concert/ConcertSync/ConcertSyncCore/Source/ConcertSyncCore/Public/HistoryEdition/ActivityNode.h b/Engine/Plugins/Developer/Concert/ConcertSync/ConcertSyncCore/Source/ConcertSyncCore/Public/HistoryEdition/ActivityNode.h
--- a/Engine/Plugins/Developer/Concert/ConcertSync/ConcertSyncCore/Source/ConcertSyncCore/Public/HistoryEdition/ActivityNode.h
+++ b/Engine/Plugins/Developer/Concert/ConcertSync/ConcertSyncCore/Source/ConcertSyncCore/Public/HistoryEdition/ActivityNode.h
@@ -12,15 +12,25 @@ namespace UE::ConcertSyncCore
 {
 	class FActivityDependencyGraph;
 
+	enum class EActivityNodeFlags
+	{
+		None			= 0x00,
+
+		/** This activity renames a package */
+		RenameActivity	= 0x01
+	};
+	ENUM_CLASS_FLAGS(EActivityNodeFlags)
+
 	/** A node corresponds to an activity and can depend on preceding activity*/
 	class CONCERTSYNCCORE_API FActivityNode
 	{
 		friend FActivityDependencyGraph;
 	public:
 		
-		FActivityNode(int64 ActivityId, FActivityNodeID NodeIndex)
+		FActivityNode(int64 ActivityId, FActivityNodeID NodeIndex, EActivityNodeFlags NodeFlags)
 			: ActivityId(ActivityId)
 			, NodeIndex(NodeIndex)
+			, NodeFlags(NodeFlags)
 		{}
 
 		bool HasAnyDependency() const { return Dependencies.Num() > 0; }
@@ -33,6 +43,7 @@ namespace UE::ConcertSyncCore
 		
 		int64 GetActivityId() const { return ActivityId; }
 		FActivityNodeID GetNodeIndex() const { return NodeIndex; }
+		EActivityNodeFlags GetNodeFlags() const { return NodeFlags; }
 		const TArray<FActivityDependencyEdge>& GetDependencies() const { return Dependencies; }
 		const TArray<FActivityNodeID>& GetAffectedChildren() const { return AffectedChildren; }
 	
@@ -41,9 +52,12 @@ namespace UE::ConcertSyncCore
 		/** The activity this node corresponds to */
 		const int64 ActivityId;
 
-		/** Index in parent tree's node array (to avoid bad memory footprint) */
+		/** Index in parent tree's node array (graph uses TArray to avoid bad memory footprint) */
 		const FActivityNodeID NodeIndex;
 
+		/** Additional info about this activity node */
+		const EActivityNodeFlags NodeFlags;
+
 		/** This node's parents. Unset implies this is the root node. */
 		TArray<FActivityDependencyEdge> Dependencies;
 
diff --git a/Engine/Plugins/Developer/Concert/ConcertSync/ConcertSyncTest/Source/ConcertSyncTest/Private/HistoryEdition/DependencyGraphBuilderTests.cpp b/Engine/Plugins/Developer/Concert/ConcertSync/ConcertSyncTest/Source/ConcertSyncTest/Private/HistoryEdition/DependencyGraphBuilderTests.cpp
--- a/Engine/Plugins/Developer/Concert/ConcertSync/ConcertSyncTest/Source/ConcertSyncTest/Private/HistoryEdition/DependencyGraphBuilderTests.cpp
+++ b/Engine/Plugins/Developer/Concert/ConcertSync/ConcertSyncTest/Source/ConcertSyncTest/Private/HistoryEdition/DependencyGraphBuilderTests.cpp
@@ -60,14 +60,17 @@ namespace UE::ConcertSyncTests::RenameEditAndDeleteMapsFlowTest
 		// 1 Create map Foo
 		{
 			Test.TestFalse(TEXT("1 Creating new package 'Foo' has no dependencies."), Graph.GetNodeById(ActivityNodes[_1_NewPackageFoo]).HasAnyDependency());
+			Test.TestTrue(TEXT("1 Creating new package 'Foo' has correct node flags"), Graph.GetNodeById(ActivityNodes[_1_NewPackageFoo]).GetNodeFlags() == EActivityNodeFlags::None);
 			
 			Test.TestTrue(TEXT("1 Saving 'Foo' has dependency to creating package 'Foo'."), Graph.GetNodeById(ActivityNodes[_1_SavePackageFoo]).DependsOnActivity(ActivityMappings[_1_NewPackageFoo], Graph, EActivityDependencyReason::PackageCreation, EDependencyStrength::HardDependency));
+			Test.TestTrue(TEXT("1 Saving 'Foo' has correct node flags"), Graph.GetNodeById(ActivityNodes[_1_SavePackageFoo]).GetNodeFlags() == EActivityNodeFlags::None);
 			Test.TestEqual(TEXT("1 Saving 'Foo' has exactly 1 dependency"), Graph.GetNodeById(ActivityNodes[_1_SavePackageFoo]).GetDependencies().Num(), 1);
 		}
 
 		// 2 Add actor A
 		{
 			Test.TestTrue(TEXT("2 Adding actor to 'Foo' depends on creating package 'Foo'."), Graph.GetNodeById(ActivityNodes[_2_AddActor]).DependsOnActivity(ActivityMappings[_1_NewPackageFoo], Graph, EActivityDependencyReason::PackageCreation, EDependencyStrength::HardDependency));
+			Test.TestTrue(TEXT("2 Adding actor to 'Foo' has correct node flags"), Graph.GetNodeById(ActivityNodes[_2_AddActor]).GetNodeFlags() == EActivityNodeFlags::None);
 			Test.TestEqual(TEXT("2 Saving 'Foo' has exactly 1 dependency"), Graph.GetNodeById(ActivityNodes[_2_AddActor]).GetDependencies().Num(), 1);
 		}
 
@@ -75,6 +78,7 @@ namespace UE::ConcertSyncTests::RenameEditAndDeleteMapsFlowTest
 		{
 			// It must be a EDependencyStrength::HardDependency because you cannot edit the actor without having created it 
 			Test.TestTrue(TEXT("3 Renaming actor depends on having created the actor."), Graph.GetNodeById(ActivityNodes[_3_RenameActor]).DependsOnActivity(ActivityMappings[_2_AddActor], Graph, EActivityDependencyReason::SubobjectCreation, EDependencyStrength::HardDependency));
+			Test.TestTrue(TEXT("3 Renaming actor to 'Foo' has correct node flags"), Graph.GetNodeById(ActivityNodes[_3_RenameActor]).GetNodeFlags() == EActivityNodeFlags::None);
 			Test.TestEqual(TEXT("3 Renaming actor has exactly 1 dependency"), Graph.GetNodeById(ActivityNodes[_3_RenameActor]).GetDependencies().Num(), 1);
 		}
 
@@ -85,36 +89,43 @@ namespace UE::ConcertSyncTests::RenameEditAndDeleteMapsFlowTest
 			Test.TestTrue(TEXT("4 Editing actor may depend on having edited the actor previously."), Graph.GetNodeById(ActivityNodes[_4_EditActor]).DependsOnActivity(ActivityMappings[_3_RenameActor], Graph, EActivityDependencyReason::EditAfterPreviousPackageEdit, EDependencyStrength::PossibleDependency));
 			// This activity must have a EDependencyStrength::HardDependency to _2_AddActor because the edit cannot happen without having created the actor
 			Test.TestTrue(TEXT("4 Editing actor depends on having created the actor."), Graph.GetNodeById(ActivityNodes[_4_EditActor]).DependsOnActivity(ActivityMappings[_2_AddActor], Graph, EActivityDependencyReason::SubobjectCreation, EDependencyStrength::HardDependency));
+			Test.TestTrue(TEXT("4 Editing actor has correct node flags"), Graph.GetNodeById(ActivityNodes[_4_EditActor]).GetNodeFlags() == EActivityNodeFlags::None);
 			Test.TestEqual(TEXT("4 Editing actor has exactly 1 dependency"), Graph.GetNodeById(ActivityNodes[_4_EditActor]).GetDependencies().Num(), 2);
 		}
 
 		// 5 Rename map to Bar
 		{
 			Test.TestFalse(TEXT("5 Saving new package 'Bar' has no dependencies."), Graph.GetNodeById(ActivityNodes[_5_SavePackageBar]).HasAnyDependency());
+			Test.TestTrue(TEXT("5 Saving new package 'Bar' has correct node flags"), Graph.GetNodeById(ActivityNodes[_5_SavePackageBar]).GetNodeFlags() == EActivityNodeFlags::None);
 			
 			Test.TestTrue(TEXT("5 Renaming 'Foo' to 'Bar' has dependency to creating package 'Foo'."), Graph.GetNodeById(ActivityNodes[_5_RenameFooToBar]).DependsOnActivity(ActivityMappings[_1_NewPackageFoo], Graph, EActivityDependencyReason::PackageCreation, EDependencyStrength::HardDependency));
 			Test.TestTrue(TEXT("5 Renaming 'Foo' to 'Bar' has dependency to creating package 'Bar'."), Graph.GetNodeById(ActivityNodes[_5_RenameFooToBar]).DependsOnActivity(ActivityMappings[_5_SavePackageBar], Graph, EActivityDependencyReason::PackageCreation, EDependencyStrength::PossibleDependency));
+			Test.TestTrue(TEXT("5 Renaming 'Foo' to 'Bar' has correct node flags"), Graph.GetNodeById(ActivityNodes[_5_RenameFooToBar]).GetNodeFlags() == EActivityNodeFlags::RenameActivity);
 			Test.TestEqual(TEXT("5 Saving 'Foo' has exactly 2 dependencies"), Graph.GetNodeById(ActivityNodes[_5_RenameFooToBar]).GetDependencies().Num(), 2);
 		}
 
 		// 6 Edit actor A
 		{
 			Test.TestTrue(TEXT("6 Editing actor in 'Bar' depends on having renamed 'Foo' to 'Bar'."), Graph.GetNodeById(ActivityNodes[_6_EditActor]).DependsOnActivity(ActivityMappings[_5_RenameFooToBar], Graph, EActivityDependencyReason::PackageRename, EDependencyStrength::HardDependency));
+			Test.TestTrue(TEXT("6 Editing actor has correct node flags"), Graph.GetNodeById(ActivityNodes[_6_EditActor]).GetNodeFlags() == EActivityNodeFlags::None);
 			Test.TestEqual(TEXT("6 Editing actor has exactly 1 dependency"), Graph.GetNodeById(ActivityNodes[_6_EditActor]).GetDependencies().Num(), 1);
 		}
 
 		// 7 Delete map Bar
 		{
 			Test.TestTrue(TEXT("7 Deleting 'Bar' depends on having renamed 'Foo' to 'Bar' previously."), Graph.GetNodeById(ActivityNodes[_7_DeleteBar]).DependsOnActivity(ActivityMappings[_5_RenameFooToBar], Graph, EActivityDependencyReason::PackageRename, EDependencyStrength::HardDependency));
+			Test.TestTrue(TEXT("7 Deleting 'Bar' has correct node flags"), Graph.GetNodeById(ActivityNodes[_7_DeleteBar]).GetNodeFlags() == EActivityNodeFlags::None);
 			Test.TestEqual(TEXT("7 Deleting 'Bar' after rename has exactly 1 dependency"), Graph.GetNodeById(ActivityNodes[_7_DeleteBar]).GetDependencies().Num(), 1);
 		}
 
 		// 8 Create map Bar
 		{
 			Test.TestTrue(TEXT("8 Re-creating 'Bar' depends on having deleted 'Bar' previously."), Graph.GetNodeById(ActivityNodes[_8_NewPackageFoo]).DependsOnActivity(ActivityMappings[_7_DeleteBar], Graph, EActivityDependencyReason::PackageRemoval, EDependencyStrength::HardDependency));
+			Test.TestTrue(TEXT("8 Re-creating actor has correct node flags"), Graph.GetNodeById(ActivityNodes[_8_NewPackageFoo]).GetNodeFlags() == EActivityNodeFlags::None);
 			Test.TestEqual(TEXT("8 Re-creating 'Bar' 1 dependency"), Graph.GetNodeById(ActivityNodes[_8_NewPackageFoo]).GetDependencies().Num(), 1);
 			
 			Test.TestTrue(TEXT("8 Saving 'Bar' depends on re-created 'Bar'."), Graph.GetNodeById(ActivityNodes[_8_SavePackageFoo]).DependsOnActivity(ActivityMappings[_8_NewPackageFoo], Graph, EActivityDependencyReason::PackageCreation, EDependencyStrength::HardDependency));
+			Test.TestTrue(TEXT("8 Saving 'Bar' has correct node flags"), Graph.GetNodeById(ActivityNodes[_8_SavePackageFoo]).GetNodeFlags() == EActivityNodeFlags::None);
 			Test.TestEqual(TEXT("8 Saving 'Bar' after re-creation has exactly 1 dependency"), Graph.GetNodeById(ActivityNodes[_8_SavePackageFoo]).GetDependencies().Num(), 1);
 		}
 		
diff --git a/Engine/Plugins/Developer/Concert/ConcertSync/ConcertSyncTest/Source/ConcertSyncTest/Private/HistoryEdition/HistoryAnalysisTests.cpp b/Engine/Plugins/Developer/Concert/ConcertSync/ConcertSyncTest/Source/ConcertSyncTest/Private/HistoryEdition/HistoryAnalysisTests.cpp
--- a/Engine/Plugins/Developer/Concert/ConcertSync/ConcertSyncTest/Source/ConcertSyncTest/Private/HistoryEdition/HistoryAnalysisTests.cpp
+++ b/Engine/Plugins/Developer/Concert/ConcertSync/ConcertSyncTest/Source/ConcertSyncTest/Private/HistoryEdition/HistoryAnalysisTests.cpp
@@ -80,13 +80,15 @@ namespace UE::ConcertSyncTests::AnalysisTests
 			 * 
 			 * Note: The transaction activities (_3_EditActor, _4_EditActor) have possible dependencies BUT they do have hard dependencies to _2_AddActor. This is why they must be in HardDependencies, too.
 			 */
-			const TSet<ETestActivity> ExcludedActivities{ _1_NewPackageFoo, _5_SavePackageBar };
-			const TSet<ETestActivity> HardDependencies = AllActivities().Difference(ExcludedActivities);
-			const bool bDeleteAllCorrect = ValidateRequirements(TEXT("Delete /Game/Foo"), *this, Activities, DeleteFooRequirements, HardDependencies);
+			const TSet<ETestActivity> ExcludedHardActivities { _1_NewPackageFoo, _5_SavePackageBar };
+			const TSet<ETestActivity> HardDependencies = AllActivities().Difference(ExcludedHardActivities);
+			const TSet<ETestActivity> PossibleDependencies { _5_SavePackageBar };
+			const bool bDeleteAllCorrect = ValidateRequirements(TEXT("Delete /Game/Foo"), *this, Activities, DeleteFooRequirements, HardDependencies, PossibleDependencies);
 			TestTrue(TEXT("Delete /Game/Foo is correct"), bDeleteAllCorrect);
 		}
 
 		// Delete rename transaction > No dependencies
+			// because the rename basically just changes AActor::ActorLabel
 		{
 			const FHistoryDeletionRequirements DeleteRenameRequirements = AnalyseActivityDeletion({ Activities[_3_RenameActor] }, DependencyGraph);
 			TestEqual(TEXT("Delete renaming actor: HardDependencies.Num() == 0"), DeleteRenameRequirements.HardDependencies.Num(), 0);
@@ -103,6 +105,23 @@ namespace UE::ConcertSyncTests::AnalysisTests
 			TestTrue(TEXT("Delete actor creation: Edit depends on created actor"), DeleteCreateActorRequirements.HardDependencies.Contains(Activities[_4_EditActor]));
 		}
 
+		// Deleting a rename activity > Rename activity should have a possible dependency to activities that saved the renamed-to package
+			// because MU internally creates a Save Package activity before a rename activity (so secretly rename activities is two related activities)
+		{
+			const FHistoryDeletionRequirements DeleteRenameRequirements = AnalyseActivityDeletion({ Activities[_5_RenameFooToBar] }, DependencyGraph);
+
+			// Not really what we want to test - but we're doing it for completness
+			TestEqual(TEXT("Delete rename package: HardDependencies.Num() == 4"), DeleteRenameRequirements.HardDependencies.Num(), 4);
+			TestTrue(TEXT("Delete rename package: Rename has hard dependency creating renamed package"), DeleteRenameRequirements.HardDependencies.Contains(Activities[_6_EditActor]));
+			TestTrue(TEXT("Delete rename package: Rename has hard dependency to editing renamed package (1)"), DeleteRenameRequirements.HardDependencies.Contains(Activities[_7_DeleteBar]));
+			TestTrue(TEXT("Delete rename package: Rename has hard dependency to editing renamed package (2)"), DeleteRenameRequirements.HardDependencies.Contains(Activities[_8_NewPackageFoo]));
+			TestTrue(TEXT("Delete rename package: Rename has hard dependency to editing renamed package (3)"), DeleteRenameRequirements.HardDependencies.Contains(Activities[_8_SavePackageFoo]));
+
+			// Now what we really wanted to test
+			TestEqual(TEXT("Delete rename package: PossibleDependencies.Num() == 1"), DeleteRenameRequirements.PossibleDependencies.Num(), 1);
+			TestTrue(TEXT("Delete rename package: Rename has soft dependency to saving renamed-to package"), DeleteRenameRequirements.PossibleDependencies.Contains(Activities[_5_SavePackageBar]));
+		}
+
 		return true;
 	}
 
diff --git a/Engine/Plugins/Developer/Concert/ConcertSync/ConcertSyncTest/Source/ConcertSyncTest/Private/HistoryEdition/RenameEditAndDeleteMapsFlow.h b/Engine/Plugins/Developer/Concert/ConcertSync/ConcertSyncTest/Source/ConcertSyncTest/Private/HistoryEdition/RenameEditAndDeleteMapsFlow.h
--- a/Engine/Plugins/Developer/Concert/ConcertSync/ConcertSyncTest/Source/ConcertSyncTest/Private/HistoryEdition/RenameEditAndDeleteMapsFlow.h
+++ b/Engine/Plugins/Developer/Concert/ConcertSync/ConcertSyncTest/Source/ConcertSyncTest/Private/HistoryEdition/RenameEditAndDeleteMapsFlow.h
@@ -11,17 +11,18 @@ namespace UE::ConcertSyncTests::RenameEditAndDeleteMapsFlowTest
 {
 	enum ETestActivity
 	{
-		_1_NewPackageFoo,
-		_1_SavePackageFoo,
-		_2_AddActor,
-		_3_RenameActor,
-		_4_EditActor,
-		_5_SavePackageBar,
-		_5_RenameFooToBar,
-		_6_EditActor,
-		_7_DeleteBar,
-		_8_NewPackageFoo,
-		_8_SavePackageFoo,
+		// Redundant = are visual finding aid during debugging
+		_1_NewPackageFoo	= 0,
+		_1_SavePackageFoo	= 1,
+		_2_AddActor			= 2,
+		_3_RenameActor		= 3,
+		_4_EditActor		= 4,
+		_5_SavePackageBar	= 5,
+		_5_RenameFooToBar	= 6,
+		_6_EditActor		= 7,
+		_7_DeleteBar		= 8,
+		_8_NewPackageFoo	= 9,
+		_8_SavePackageFoo	= 10,
 		
 		ActivityCount
 	};
