Get message address for concert client

diff --git a/Engine/Plugins/Developer/Concert/ConcertMain/Source/Concert/Private/ConcertServerSession.h b/Engine/Plugins/Developer/Concert/ConcertMain/Source/Concert/Private/ConcertServerSession.h
--- a/Engine/Plugins/Developer/Concert/ConcertMain/Source/Concert/Private/ConcertServerSession.h
+++ b/Engine/Plugins/Developer/Concert/ConcertMain/Source/Concert/Private/ConcertServerSession.h
@@ -39,6 +39,11 @@ public:
 		SendSessionNameChanged();
 	}
 
+	virtual FMessageAddress GetClientAddress(const FGuid& ClientEndpointId) const override
+	{
+		return ServerSessionEndpoint->GetRemoteAddress(ClientEndpointId);
+	}
+
 	void SetLastModifiedToNow()
 	{
 		SessionInfo.SetLastModifiedToNow();
diff --git a/Engine/Plugins/Developer/Concert/ConcertMain/Source/Concert/Public/IConcertSession.h b/Engine/Plugins/Developer/Concert/ConcertMain/Source/Concert/Public/IConcertSession.h
--- a/Engine/Plugins/Developer/Concert/ConcertMain/Source/Concert/Public/IConcertSession.h
+++ b/Engine/Plugins/Developer/Concert/ConcertMain/Source/Concert/Public/IConcertSession.h
@@ -247,6 +247,9 @@ public:
 	/** Rename the session. */
 	virtual void SetName(const FString& NewName) = 0;
 
+	/** Gets the client's address used by the messaging system. Can be fed into  */
+	virtual FMessageAddress GetClientAddress(const FGuid& ClientEndpointId) const = 0;
+
 	/** Callback when a server session gets ticked */
 	virtual FOnConcertServerSessionTick& OnTick() = 0;
 
diff --git a/Engine/Plugins/Developer/Concert/ConcertMain/Source/ConcertTransport/Private/ConcertLocalEndpoint.cpp b/Engine/Plugins/Developer/Concert/ConcertMain/Source/ConcertTransport/Private/ConcertLocalEndpoint.cpp
--- a/Engine/Plugins/Developer/Concert/ConcertMain/Source/ConcertTransport/Private/ConcertLocalEndpoint.cpp
+++ b/Engine/Plugins/Developer/Concert/ConcertMain/Source/ConcertTransport/Private/ConcertLocalEndpoint.cpp
@@ -137,6 +137,14 @@ const FConcertEndpointContext& FConcertLocalEndpoint::GetEndpointContext() const
 	return EndpointContext;
 }
 
+FMessageAddress FConcertLocalEndpoint::GetRemoteAddress(const FGuid& ConcertEndpointId) const
+{
+	const FConcertRemoteEndpointPtr RemoteEndpoint = FindRemoteEndpoint(ConcertEndpointId);
+	return RemoteEndpoint
+		? RemoteEndpoint->GetAddress()
+		: FMessageAddress();
+}
+
 FOnConcertRemoteEndpointConnectionChanged& FConcertLocalEndpoint::OnRemoteEndpointConnectionChanged()
 {
 	return OnRemoteEndpointConnectionChangedDelegate;
diff --git a/Engine/Plugins/Developer/Concert/ConcertMain/Source/ConcertTransport/Private/ConcertLocalEndpoint.h b/Engine/Plugins/Developer/Concert/ConcertMain/Source/ConcertTransport/Private/ConcertLocalEndpoint.h
--- a/Engine/Plugins/Developer/Concert/ConcertMain/Source/ConcertTransport/Private/ConcertLocalEndpoint.h
+++ b/Engine/Plugins/Developer/Concert/ConcertMain/Source/ConcertTransport/Private/ConcertLocalEndpoint.h
@@ -31,6 +31,7 @@ public:
 	virtual ~FConcertLocalEndpoint();
 
 	virtual const FConcertEndpointContext& GetEndpointContext() const override;
+	virtual FMessageAddress GetRemoteAddress(const FGuid& ConcertEndpointId) const override;
 
 	virtual FOnConcertRemoteEndpointConnectionChanged& OnRemoteEndpointConnectionChanged() override;
 
diff --git a/Engine/Plugins/Developer/Concert/ConcertMain/Source/ConcertTransport/Public/IConcertEndpoint.h b/Engine/Plugins/Developer/Concert/ConcertMain/Source/ConcertTransport/Public/IConcertEndpoint.h
--- a/Engine/Plugins/Developer/Concert/ConcertMain/Source/ConcertTransport/Public/IConcertEndpoint.h
+++ b/Engine/Plugins/Developer/Concert/ConcertMain/Source/ConcertTransport/Public/IConcertEndpoint.h
@@ -86,6 +86,9 @@ public:
 	/** Get the context for this endpoint. */
 	virtual const FConcertEndpointContext& GetEndpointContext() const = 0;
 
+	/** Gets the remote address of a client or the server. If ConcertEndpointId is not valid, the result's IsValid return false. */
+	virtual FMessageAddress GetRemoteAddress(const FGuid& ConcertEndpointId) const = 0;
+
 	/** Callback when a remote endpoint connection changes. */
 	virtual FOnConcertRemoteEndpointConnectionChanged& OnRemoteEndpointConnectionChanged() = 0;
 
diff --git a/Engine/Plugins/Developer/Concert/ConcertSync/ConcertSyncTest/Source/ConcertSyncTest/Private/ConcertDataStoreTests.cpp b/Engine/Plugins/Developer/Concert/ConcertSync/ConcertSyncTest/Source/ConcertSyncTest/Private/ConcertDataStoreTests.cpp
--- a/Engine/Plugins/Developer/Concert/ConcertSync/ConcertSyncTest/Source/ConcertSyncTest/Private/ConcertDataStoreTests.cpp
+++ b/Engine/Plugins/Developer/Concert/ConcertSync/ConcertSyncTest/Source/ConcertSyncTest/Private/ConcertDataStoreTests.cpp
@@ -80,6 +80,7 @@ public:
 	virtual const FGuid& GetId() const override																								{ return Id; }
 	virtual const FString& GetName() const override																							{ return NotMocked<const FString&>(Name); }
 	virtual void SetName(const FString&) override																							{ return NotMocked<void>(); }
+	virtual FMessageAddress GetClientAddress(const FGuid& ClientEndpointId) const override													{ return NotMocked<FMessageAddress>(); }
 	virtual const FConcertSessionInfo& GetSessionInfo() const override																		{ return NotMocked<const FConcertSessionInfo&>(SessionInfo); }
 	virtual FString GetSessionWorkingDirectory() const override																				{ return GetTestSessionRootPath() / Id.ToString(); }
 	virtual TArray<FGuid> GetSessionClientEndpointIds() const override																		{ return NotMocked<TArray<FGuid>>(); }
diff --git a/Engine/Plugins/Messaging/UdpMessaging/Source/UdpMessaging/Private/UdpMessagingModule.cpp b/Engine/Plugins/Messaging/UdpMessaging/Source/UdpMessaging/Private/UdpMessagingModule.cpp
--- a/Engine/Plugins/Messaging/UdpMessaging/Source/UdpMessaging/Private/UdpMessagingModule.cpp
+++ b/Engine/Plugins/Messaging/UdpMessaging/Source/UdpMessaging/Private/UdpMessagingModule.cpp
@@ -392,6 +392,11 @@ public:
 		return {};
 	}
 
+	virtual FGuid GetNodeIdFromAddress(const FMessageAddress& MessageAddress) const override
+	{
+		return MessageBridge->LookupAddress(MessageAddress);
+	}
+
 	virtual FOnTransferDataUpdated& OnTransferUpdatedFromThread() override
 	{
 		return UE::Private::MessageProcessor::OnSegmenterUpdated();
diff --git a/Engine/Source/Runtime/Messaging/Public/INetworkMessagingExtension.h b/Engine/Source/Runtime/Messaging/Public/INetworkMessagingExtension.h
--- a/Engine/Source/Runtime/Messaging/Public/INetworkMessagingExtension.h
+++ b/Engine/Source/Runtime/Messaging/Public/INetworkMessagingExtension.h
@@ -8,6 +8,8 @@
 #include "Misc/Guid.h"
 #include "Misc/Timespan.h"
 
+struct FMessageAddress;
+
 /** General information on transport health per Node Id. */
 struct FMessageTransportStatistics
 {
@@ -115,6 +117,11 @@
 	virtual FMessageTransportStatistics GetLatestNetworkStatistics(FGuid NodeId) const = 0;
 
 	/**
+	 * Return the node ID corresponding to MessageAddress
+	 */
+	virtual FGuid GetNodeIdFromAddress(const FMessageAddress& MessageAddress) const = 0;
+
+	/**
 	 * Delegate invoked when any transmission statistics are updated. This delegate may get called from another thread.
 	 * Please consider thread safety when receiving this delegate.
 	 */
