Fix failing to delete activities if there is a live session with the same name as the archived session. Also some naming changes for upcoming muting changes.

diff --git a/Engine/Plugins/Developer/Concert/ConcertApp/MultiUserServer/Source/MultiUserServer/Private/Widgets/HistoryDeletion/SActivityDependencyView.cpp b/Engine/Plugins/Developer/Concert/ConcertApp/MultiUserServer/Source/MultiUserServer/Private/Widgets/HistoryDeletion/SActivityDependencyView.cpp
--- a/Engine/Plugins/Developer/Concert/ConcertApp/MultiUserServer/Source/MultiUserServer/Private/Widgets/HistoryDeletion/SActivityDependencyView.cpp
+++ b/Engine/Plugins/Developer/Concert/ConcertApp/MultiUserServer/Source/MultiUserServer/Private/Widgets/HistoryDeletion/SActivityDependencyView.cpp
@@ -30,7 +30,7 @@ namespace UE::MultiUserServer
 	}
 }
 
-void SActivityDependencyView::Construct(const FArguments& InArgs, const UE::ConcertSyncCore::FHistoryDeletionRequirements& DeletionRequirements)
+void SActivityDependencyView::Construct(const FArguments& InArgs, const UE::ConcertSyncCore::FHistoryEditionArgs& DeletionRequirements)
 {
 	check(InArgs._CreateSessionHistory.IsBound());
 
diff --git a/Engine/Plugins/Developer/Concert/ConcertApp/MultiUserServer/Source/MultiUserServer/Private/Widgets/HistoryDeletion/SActivityDependencyView.h b/Engine/Plugins/Developer/Concert/ConcertApp/MultiUserServer/Source/MultiUserServer/Private/Widgets/HistoryDeletion/SActivityDependencyView.h
--- a/Engine/Plugins/Developer/Concert/ConcertApp/MultiUserServer/Source/MultiUserServer/Private/Widgets/HistoryDeletion/SActivityDependencyView.h
+++ b/Engine/Plugins/Developer/Concert/ConcertApp/MultiUserServer/Source/MultiUserServer/Private/Widgets/HistoryDeletion/SActivityDependencyView.h
@@ -12,7 +12,7 @@ namespace UE
 {
 	namespace ConcertSyncCore
 	{
-		struct FHistoryDeletionRequirements;
+		struct FHistoryEditionArgs;
 	}
 }
 
@@ -33,5 +33,5 @@ public:
 		SLATE_ARGUMENT(FText, CustomActivityColumnLabel)
     SLATE_END_ARGS()
 	
-	void Construct(const FArguments& InArgs, const UE::ConcertSyncCore::FHistoryDeletionRequirements& DeletionRequirements);
+	void Construct(const FArguments& InArgs, const UE::ConcertSyncCore::FHistoryEditionArgs& DeletionRequirements);
 };
diff --git a/Engine/Plugins/Developer/Concert/ConcertApp/MultiUserServer/Source/MultiUserServer/Private/Widgets/HistoryDeletion/SDeleteActivityDependenciesDialog.cpp b/Engine/Plugins/Developer/Concert/ConcertApp/MultiUserServer/Source/MultiUserServer/Private/Widgets/HistoryDeletion/SDeleteActivityDependenciesDialog.cpp
--- a/Engine/Plugins/Developer/Concert/ConcertApp/MultiUserServer/Source/MultiUserServer/Private/Widgets/HistoryDeletion/SDeleteActivityDependenciesDialog.cpp
+++ b/Engine/Plugins/Developer/Concert/ConcertApp/MultiUserServer/Source/MultiUserServer/Private/Widgets/HistoryDeletion/SDeleteActivityDependenciesDialog.cpp
@@ -12,7 +12,7 @@
 
 #define LOCTEXT_NAMESPACE "UnrealMultiUserUI"
 
-void SDeleteActivityDependenciesDialog::Construct(const FArguments& InArgs, const FGuid& SessionId, const TSharedRef<IConcertSyncServer>& SyncServer, UE::ConcertSyncCore::FHistoryDeletionRequirements InDeletionRequirements)
+void SDeleteActivityDependenciesDialog::Construct(const FArguments& InArgs, const FGuid& SessionId, const TSharedRef<IConcertSyncServer>& SyncServer, UE::ConcertSyncCore::FHistoryEditionArgs InDeletionRequirements)
 {
 	DeletionRequirements = MoveTemp(InDeletionRequirements);
 	OnConfirmDeletionFunc = InArgs._OnConfirmDeletion;
@@ -52,7 +52,7 @@ void SDeleteActivityDependenciesDialog::Construct(const FArguments& InArgs, cons
 	);
 }
 
-TSharedRef<SWidget> SDeleteActivityDependenciesDialog::CreateBody(const FGuid& InSessionId, const TSharedRef<IConcertSyncServer>& InSyncServer, UE::ConcertSyncCore::FHistoryDeletionRequirements InDeletionRequirements)
+TSharedRef<SWidget> SDeleteActivityDependenciesDialog::CreateBody(const FGuid& InSessionId, const TSharedRef<IConcertSyncServer>& InSyncServer, UE::ConcertSyncCore::FHistoryEditionArgs InDeletionRequirements)
 {
 	for (FActivityID HardDependency : DeletionRequirements.HardDependencies)
 	{
diff --git a/Engine/Plugins/Developer/Concert/ConcertApp/MultiUserServer/Source/MultiUserServer/Private/Widgets/HistoryDeletion/SDeleteActivityDependenciesDialog.h b/Engine/Plugins/Developer/Concert/ConcertApp/MultiUserServer/Source/MultiUserServer/Private/Widgets/HistoryDeletion/SDeleteActivityDependenciesDialog.h
--- a/Engine/Plugins/Developer/Concert/ConcertApp/MultiUserServer/Source/MultiUserServer/Private/Widgets/HistoryDeletion/SDeleteActivityDependenciesDialog.h
+++ b/Engine/Plugins/Developer/Concert/ConcertApp/MultiUserServer/Source/MultiUserServer/Private/Widgets/HistoryDeletion/SDeleteActivityDependenciesDialog.h
@@ -5,7 +5,6 @@
 #include "CoreMinimal.h"
 #include "ConcertMessageData.h"
 #include "IConcertSyncServer.h"
-
 #include "Dialog/SCustomDialog.h"
 #include "HistoryEdition/HistoryAnalysis.h"
 #include "Widgets/DeclarativeSyntaxSupport.h"
@@ -17,7 +16,7 @@ class SDeleteActivityDependenciesDialog : public SCustomDialog
 {
 public:
 
-	DECLARE_DELEGATE_OneParam(FConfirmDeletion, const UE::ConcertSyncCore::FHistoryDeletionRequirements& /*SelectedRequirements*/)
+	DECLARE_DELEGATE_OneParam(FConfirmDeletion, const UE::ConcertSyncCore::FHistoryEditionArgs& /*SelectedRequirements*/)
 
 	SLATE_BEGIN_ARGS(SDeleteActivityDependenciesDialog)
 	{}
@@ -28,11 +27,11 @@ public:
 	/**
 	 * @param InDeletionRequirements Specifies which activities must be deleted and which are optional.
 	 */
-	void Construct(const FArguments& InArgs, const FGuid& SessionId, const TSharedRef<IConcertSyncServer>& SyncServer, UE::ConcertSyncCore::FHistoryDeletionRequirements InDeletionRequirements);
+	void Construct(const FArguments& InArgs, const FGuid& SessionId, const TSharedRef<IConcertSyncServer>& SyncServer, UE::ConcertSyncCore::FHistoryEditionArgs InDeletionRequirements);
 
 private:
 	
-	UE::ConcertSyncCore::FHistoryDeletionRequirements DeletionRequirements;
+	UE::ConcertSyncCore::FHistoryEditionArgs DeletionRequirements;
 	FConfirmDeletion OnConfirmDeletionFunc;
 	
 	/** Maps each activity with a possible dependency to whether it will be deleted or not. */
@@ -41,7 +40,7 @@ private:
 	/** Filters out all activities that are not part of the deletion requirements. */
 	TSharedPtr<FArchivedSessionHistoryController> FilteredSessionHistoryController;
 
-	TSharedRef<SWidget> CreateBody(const FGuid& InSessionId, const TSharedRef<IConcertSyncServer>& InSyncServer, UE::ConcertSyncCore::FHistoryDeletionRequirements InDeletionRequirements);
+	TSharedRef<SWidget> CreateBody(const FGuid& InSessionId, const TSharedRef<IConcertSyncServer>& InSyncServer, UE::ConcertSyncCore::FHistoryEditionArgs InDeletionRequirements);
 
 	void OnConfirmPressed();
 };
diff --git a/Engine/Plugins/Developer/Concert/ConcertApp/MultiUserServer/Source/MultiUserServer/Private/Widgets/SessionTabs/Archived/ArchivedConcertSessionTab.cpp b/Engine/Plugins/Developer/Concert/ConcertApp/MultiUserServer/Source/MultiUserServer/Private/Widgets/SessionTabs/Archived/ArchivedConcertSessionTab.cpp
--- a/Engine/Plugins/Developer/Concert/ConcertApp/MultiUserServer/Source/MultiUserServer/Private/Widgets/SessionTabs/Archived/ArchivedConcertSessionTab.cpp
+++ b/Engine/Plugins/Developer/Concert/ConcertApp/MultiUserServer/Source/MultiUserServer/Private/Widgets/SessionTabs/Archived/ArchivedConcertSessionTab.cpp
@@ -8,7 +8,7 @@
 #include "HistoryEdition/DebugDependencyGraph.h"
 #include "HistoryEdition/DependencyGraphBuilder.h"
 #include "HistoryEdition/HistoryAnalysis.h"
-#include "HistoryEdition/HistoryDeletion.h"
+#include "HistoryEdition/HistoryEdition.h"
 #include "MultiUserServerConsoleVariables.h"
 #include "MultiUserServerModule.h"
 #include "Session/History/SEditableSessionHistory.h"
@@ -67,16 +67,16 @@ void FArchivedConcertSessionTab::OnRequestDeleteActivity(const TSet<TSharedRef<F
 		{
 			return Activity->Activity.ActivityId;
 		});
-		FHistoryDeletionRequirements DeletionRequirements = AnalyseActivityDeletion(RequestedForDelete, DependencyGraph, true);
+		FHistoryEditionArgs DeletionRequirements = AnalyseActivityDependencies(RequestedForDelete, DependencyGraph, true);
 
 		TWeakPtr<const FArchivedConcertSessionTab> WeakTabThis = SharedThis(this);
 		TSharedRef<SDeleteActivityDependenciesDialog> Dialog = SNew(SDeleteActivityDependenciesDialog, InspectedSessionID, SyncServer, MoveTemp(DeletionRequirements))
-			.OnConfirmDeletion_Lambda([WeakTabThis](const FHistoryDeletionRequirements& SelectedRequirements)
+			.OnConfirmDeletion_Lambda([WeakTabThis](const FHistoryEditionArgs& SelectedRequirements)
 			{
 				// Because the dialog is non-modal, the user may have closed the program in the mean time
 				if (const TSharedPtr<const FArchivedConcertSessionTab> PinnedThis = WeakTabThis.Pin())
 				{
-					const FDeleteSessionErrorResult ErrorResult = DeleteActivitiesInArchivedSession(PinnedThis->SyncServer->GetConcertServer(), PinnedThis->InspectedSessionID, CombineRequirements(SelectedRequirements));
+					const FOperationErrorResult ErrorResult = DeleteActivitiesInArchivedSession(PinnedThis->SyncServer->GetConcertServer(), PinnedThis->InspectedSessionID, CombineRequirements(SelectedRequirements));
 					if (ErrorResult.HadError())
 					{
 						UE_LOG(LogConcert, Error, TEXT("Failed to delete activities from session %s: %s"), *PinnedThis->InspectedSessionID.ToString(), *ErrorResult.ErrorMessage->ToString());
diff --git a/Engine/Plugins/Developer/Concert/ConcertSync/ConcertSyncCore/Source/ConcertSyncCore/Private/HistoryEdition/HistoryAnalysis.cpp b/Engine/Plugins/Developer/Concert/ConcertSync/ConcertSyncCore/Source/ConcertSyncCore/Private/HistoryEdition/HistoryAnalysis.cpp
--- a/Engine/Plugins/Developer/Concert/ConcertSync/ConcertSyncCore/Source/ConcertSyncCore/Private/HistoryEdition/HistoryAnalysis.cpp
+++ b/Engine/Plugins/Developer/Concert/ConcertSync/ConcertSyncCore/Source/ConcertSyncCore/Private/HistoryEdition/HistoryAnalysis.cpp
@@ -12,10 +12,10 @@
 
 namespace UE::ConcertSyncCore
 {
-	FHistoryDeletionRequirements AnalyseActivityDeletion(const TSet<FActivityID>& ActivitiesToDelete, const FConcertSyncSessionDatabase& Database, bool bAddActivitiesToDelete)
+	FHistoryEditionArgs AnalyseActivityDependencies(const TSet<FActivityID>& ActivitiesToDelete, const FConcertSyncSessionDatabase& Database, bool bAddActivitiesToDelete)
 	{
 		const FActivityDependencyGraph Graph = BuildDependencyGraphFrom(Database);
-		return AnalyseActivityDeletion(ActivitiesToDelete, Graph, bAddActivitiesToDelete);
+		return AnalyseActivityDependencies(ActivitiesToDelete, Graph, bAddActivitiesToDelete);
 	}
 
 	/**
@@ -24,9 +24,9 @@ namespace UE::ConcertSyncCore
 	 */
 	static void AddSavePackageActivityAssociatedWithRenamePackageActivity(const FActivityDependencyGraph& DependencyGraph, const FActivityNode& PossibleRenameActivityNode, TSet<FActivityNodeID>& DoubleEnqueuingProtection, TQueue<FActivityNodeID>& ActivitiesToAnalyse);
 	
-	FHistoryDeletionRequirements AnalyseActivityDeletion(const TSet<FActivityID>& ActivitiesToDelete, const FActivityDependencyGraph& DependencyGraph, bool bAddActivitiesToDelete)
+	FHistoryEditionArgs AnalyseActivityDependencies(const TSet<FActivityID>& ActivitiesToEdit, const FActivityDependencyGraph& DependencyGraph, bool bAddEditedAsHardDependencies)
 	{
-		FHistoryDeletionRequirements Result;
+		FHistoryEditionArgs Result;
 
 		TSet<FActivityNodeID> PossibleDoubleEnqueuingProtection;
 		TQueue<FActivityNodeID> PossibleDependencyActivitiesToAnalyse;
@@ -34,7 +34,7 @@ namespace UE::ConcertSyncCore
 		
 		TSet<FActivityNodeID> HardDoubleEnqueuingProtection;
 		TQueue<FActivityNodeID> HardDependencyActivitiesToAnalyse;
-		for (const FActivityID ActivityToDelete : ActivitiesToDelete)
+		for (const FActivityID ActivityToDelete : ActivitiesToEdit)
 		{
 			const TOptional<FActivityNodeID> NodeID = DependencyGraph.FindNodeByActivity(ActivityToDelete);
 			if (ensureMsgf(NodeID, TEXT("Graph does not correspond to ActivitiesToDelete")))
@@ -65,7 +65,7 @@ namespace UE::ConcertSyncCore
 			const FActivityNode& ActivityNode = DependencyGraph.GetNodeById(CurrentActivityID);
 
 			const FActivityID ActivityID = ActivityNode.GetActivityId();
-			if (bAddActivitiesToDelete || !ActivitiesToDelete.Contains(ActivityID))
+			if (bAddEditedAsHardDependencies || !ActivitiesToEdit.Contains(ActivityID))
 			{
 				Result.HardDependencies.Add(ActivityID);
 			}
diff --git a/Engine/Plugins/Developer/Concert/ConcertSync/ConcertSyncCore/Source/ConcertSyncCore/Private/HistoryEdition/HistoryDeletion.cpp b/Engine/Plugins/Developer/Concert/ConcertSync/ConcertSyncCore/Source/ConcertSyncCore/Private/HistoryEdition/HistoryDeletion.cpp
deleted file mode 100644
diff --git a/Engine/Plugins/Developer/Concert/ConcertSync/ConcertSyncCore/Source/ConcertSyncCore/Private/HistoryEdition/HistoryEdition.cpp b/Engine/Plugins/Developer/Concert/ConcertSync/ConcertSyncCore/Source/ConcertSyncCore/Private/HistoryEdition/HistoryEdition.cpp
new file mode 100644
--- /dev/null
+++ b/Engine/Plugins/Developer/Concert/ConcertSync/ConcertSyncCore/Source/ConcertSyncCore/Private/HistoryEdition/HistoryEdition.cpp
@@ -0,0 +1,97 @@
+// Copyright Epic Games, Inc. All Rights Reserved.
+
+#include "HistoryEdition/HistoryEdition.h"
+
+#include "ConcertLogGlobal.h"
+#include "ConcertSyncSessionDatabase.h"
+#include "IConcertServer.h"
+
+#include "Misc/ScopeExit.h"
+
+namespace UE::ConcertSyncCore
+{
+	static FConcertSessionFilter BuildFilterFrom(const TSet<FActivityID>& ToDelete);
+
+	TSet<FActivityID> CombineRequirements(const FHistoryEditionArgs& ToDelete)
+	{
+		TSet<FActivityID> Result;
+		for (const FActivityID ActivityID : ToDelete.HardDependencies)
+		{
+			Result.Add(ActivityID);
+		}
+		for (const FActivityID ActivityID : ToDelete.PossibleDependencies)
+		{
+			Result.Add(ActivityID);
+		}
+		return Result;
+	}
+	
+	FOperationErrorResult DeleteActivitiesInArchivedSession(const TSharedRef<IConcertServer>& Server, const FGuid& ArchivedSessionId, const TSet<FActivityID>& ToDelete)
+	{
+		const FString SessionId = ArchivedSessionId.ToString(EGuidFormats::DigitsWithHyphens);
+		UE_LOG(LogConcert, Log, TEXT("Deleting %d activities from session %s..."), ToDelete.Num(), *SessionId);
+		TOptional<FString> ErrorMessage;
+		ON_SCOPE_EXIT
+		{
+			if (ErrorMessage)
+			{
+				UE_LOG(LogConcert, Error, TEXT("Finished deleting activities from %s with error: %s."), *SessionId, **ErrorMessage);
+			}
+			else
+			{
+				UE_LOG(LogConcert, Log, TEXT("Finished deleting activities from %s successfully."), *SessionId);
+			}
+		};
+		
+		const TOptional<FConcertSessionInfo> DeletedSessionInfo = Server->GetArchivedSessionInfo(ArchivedSessionId);
+		if (!DeletedSessionInfo)
+		{
+			ErrorMessage = FString::Printf(TEXT("Session ID %s does not resolve to any archived session!"), *ArchivedSessionId.ToString());
+			return FOperationErrorResult::MakeError(FText::FromString(*ErrorMessage));
+		}
+
+		// Restore the session while skipping all to be deleted activities
+		FText FailureReason;
+		const FConcertSessionFilter Filter = BuildFilterFrom(ToDelete);
+		// Do not use same name in case user already created a live session from this archived session (would fail in that case) 
+		FConcertSessionInfo OverrideInfo = *DeletedSessionInfo;
+		OverrideInfo.SessionName = DeletedSessionInfo->SessionName + FGuid::NewGuid().ToString();
+		const TSharedPtr<IConcertServerSession> LiveSession = Server->RestoreSession(ArchivedSessionId, OverrideInfo, Filter, FailureReason);
+		if (!LiveSession)
+		{
+			ErrorMessage = FailureReason.ToString();
+			return FOperationErrorResult::MakeError(FailureReason);
+		}
+		ON_SCOPE_EXIT
+		{
+			Server->DestroySession(LiveSession->GetId(), FailureReason);
+		};
+
+		// The archived session must be removed before it can be overwritten
+		if (!Server->DestroySession(ArchivedSessionId, FailureReason))
+		{
+			ErrorMessage = FailureReason.ToString();
+			return FOperationErrorResult::MakeError(FailureReason); 
+		}
+
+		if (!Server->ArchiveSession(LiveSession->GetId(), DeletedSessionInfo->SessionName, Filter, FailureReason, ArchivedSessionId).IsValid())
+		{
+			ErrorMessage = FailureReason.ToString();
+			return FOperationErrorResult::MakeError(FailureReason); 
+		}
+		
+		return FOperationErrorResult::MakeSuccess(); 
+	}
+
+	static FConcertSessionFilter BuildFilterFrom(const TSet<FActivityID>& ToDelete)
+	{
+		FConcertSessionFilter Result;
+		Result.ActivityIdsToExclude.Reserve(ToDelete.Num());
+		for (const FActivityID ActivityID : ToDelete)
+		{
+			Result.ActivityIdsToExclude.Add(ActivityID);
+		}
+		return Result;
+	}
+}
+
diff --git a/Engine/Plugins/Developer/Concert/ConcertSync/ConcertSyncCore/Source/ConcertSyncCore/Public/HistoryEdition/HistoryAnalysis.h b/Engine/Plugins/Developer/Concert/ConcertSync/ConcertSyncCore/Source/ConcertSyncCore/Public/HistoryEdition/HistoryAnalysis.h
--- a/Engine/Plugins/Developer/Concert/ConcertSync/ConcertSyncCore/Source/ConcertSyncCore/Public/HistoryEdition/HistoryAnalysis.h
+++ b/Engine/Plugins/Developer/Concert/ConcertSync/ConcertSyncCore/Source/ConcertSyncCore/Public/HistoryEdition/HistoryAnalysis.h
@@ -11,7 +11,7 @@ class FConcertSyncSessionDatabase;
 namespace UE::ConcertSyncCore
 {
 	/** Describes the activities that must be considered when deleting an activity */
-	struct CONCERTSYNCCORE_API FHistoryDeletionRequirements
+	struct CONCERTSYNCCORE_API FHistoryEditionArgs
 	{
 		/** The activities that must be removed. */
 		TSet<FActivityID> HardDependencies;
@@ -21,22 +21,22 @@ namespace UE::ConcertSyncCore
 		 */
 		TSet<FActivityID> PossibleDependencies;
 
-		FHistoryDeletionRequirements(TSet<FActivityID> HardDependencies = {}, TSet<FActivityID> PossibleDependencies = {})
+		FHistoryEditionArgs(TSet<FActivityID> HardDependencies = {}, TSet<FActivityID> PossibleDependencies = {})
 			: HardDependencies(MoveTemp(HardDependencies)),
 			  PossibleDependencies(MoveTemp(PossibleDependencies))
 		{}
 	};
 
 	/** Utility function for one-off operations: just computes the dependency graph before calling AnalyseActivityDeletion. */
-	CONCERTSYNCCORE_API FHistoryDeletionRequirements AnalyseActivityDeletion(const TSet<FActivityID>& ActivitiesToDelete, const FConcertSyncSessionDatabase& Database, bool bAddActivitiesToDelete = false);
+	CONCERTSYNCCORE_API FHistoryEditionArgs AnalyseActivityDependencies(const TSet<FActivityID>& ActivitiesToDelete, const FConcertSyncSessionDatabase& Database, bool bAddActivitiesToDelete = false);
 
 	/**
-	 * Given a set of activities to be deleted, returns which activities 1. must be and 2. may want to be deleted in addition.
-	 * @param ActivitiesToDelete The activities that should be removed
+	 * Given a set of activities to be edited (e.g. deleted or muted), returns which activities 1. must be and 2. may want to be considered in addition.
+	 * @param ActivitiesToEdit The activities that should be removed
 	 * @param DependencyGraph The graph encoding the activity dependencies
-	 * @param bAddActivitiesToDelete Whether to add ActivitiesToDelete to the result's HardDependencies
-	 * @return The activities to remove as well if ActivitiesToDelete is removed; ActivitiesToDelete is included in HardDependencies if bAddActivitiesToDelete == true.
+	 * @param bAddEditedAsHardDependencies Whether to add ActivitiesToEdit to the result's HardDependencies
+	 * @return The activities to consider as well if ActivitiesToEdit are edited; ActivitiesToEdit is included in HardDependencies if bAddActivitiesToDelete == true.
 	 */
-	CONCERTSYNCCORE_API FHistoryDeletionRequirements AnalyseActivityDeletion(const TSet<FActivityID>& ActivitiesToDelete, const FActivityDependencyGraph& DependencyGraph, bool bAddActivitiesToDelete = false);
+	CONCERTSYNCCORE_API FHistoryEditionArgs AnalyseActivityDependencies(const TSet<FActivityID>& ActivitiesToEdit, const FActivityDependencyGraph& DependencyGraph, bool bAddEditedAsHardDependencies = false);
 }
 
diff --git a/Engine/Plugins/Developer/Concert/ConcertSync/ConcertSyncCore/Source/ConcertSyncCore/Public/HistoryEdition/HistoryDeletion.h b/Engine/Plugins/Developer/Concert/ConcertSync/ConcertSyncCore/Source/ConcertSyncCore/Public/HistoryEdition/HistoryDeletion.h
deleted file mode 100644
diff --git a/Engine/Plugins/Developer/Concert/ConcertSync/ConcertSyncCore/Source/ConcertSyncCore/Public/HistoryEdition/HistoryEdition.h b/Engine/Plugins/Developer/Concert/ConcertSync/ConcertSyncCore/Source/ConcertSyncCore/Public/HistoryEdition/HistoryEdition.h
new file mode 100644
--- /dev/null
+++ b/Engine/Plugins/Developer/Concert/ConcertSync/ConcertSyncCore/Source/ConcertSyncCore/Public/HistoryEdition/HistoryEdition.h
@@ -0,0 +1,40 @@
+// Copyright Epic Games, Inc. All Rights Reserved.
+
+#pragma once
+
+#include "CoreMinimal.h"
+#include "HistoryAnalysis.h"
+
+class IConcertServer;
+class IConcertSyncServer;
+
+namespace UE::ConcertSyncCore
+{
+	struct FOperationErrorResult
+	{
+		TOptional<FText> ErrorMessage;
+
+		explicit FOperationErrorResult(TOptional<FText> ErrorMessage = {})
+			: ErrorMessage(ErrorMessage)
+		{}
+
+		static FOperationErrorResult MakeSuccess() { return FOperationErrorResult(); }
+		static FOperationErrorResult MakeError(FText Error) { return FOperationErrorResult(MoveTemp(Error)); }
+
+		bool WasSuccessful() const { return !ErrorMessage.IsSet(); }
+		bool HadError() const { return !WasSuccessful(); }
+	};
+
+	/** Utility functions that converts FHistoryDeletionRequirements into a single TSet. */
+	CONCERTSYNCCORE_API TSet<FActivityID> CombineRequirements(const FHistoryEditionArgs& ToDelete);
+
+	/**
+	 * Deletes the given sessions in ToDelete from ArchivedSessionDatabase.
+	 *
+	 * There is no direct functionality for removing activities in the database; the operation is as follows:
+	 *	1. Restore the archived session with a session filter
+	 *	2. Delete the archived session
+	 *	3. Archive the live session created in step 1 with the old guid.
+	 */
+	CONCERTSYNCCORE_API FOperationErrorResult DeleteActivitiesInArchivedSession(const TSharedRef<IConcertServer>& Server, const FGuid& ArchivedSessionId, const TSet<FActivityID>& ToDelete);
+}
diff --git a/Engine/Plugins/Developer/Concert/ConcertSync/ConcertSyncTest/Source/ConcertSyncTest/Private/HistoryEdition/HistoryAnalysisTests.cpp b/Engine/Plugins/Developer/Concert/ConcertSync/ConcertSyncTest/Source/ConcertSyncTest/Private/HistoryEdition/HistoryAnalysisTests.cpp
--- a/Engine/Plugins/Developer/Concert/ConcertSync/ConcertSyncTest/Source/ConcertSyncTest/Private/HistoryEdition/HistoryAnalysisTests.cpp
+++ b/Engine/Plugins/Developer/Concert/ConcertSync/ConcertSyncTest/Source/ConcertSyncTest/Private/HistoryEdition/HistoryAnalysisTests.cpp
@@ -22,7 +22,7 @@ namespace UE::ConcertSyncTests
 		const FString& TestBaseName,
 		FAutomationTestBase& Test,
 		const RenameEditAndDeleteMapsFlowTest::TTestActivityArray<FActivityID> Activities,
-		const ConcertSyncCore::FHistoryDeletionRequirements& ToValidate = {},
+		const ConcertSyncCore::FHistoryEditionArgs& ToValidate = {},
 		const TSet<RenameEditAndDeleteMapsFlowTest::ETestActivity>& ExpectedHardDependencies = {},
 		const TSet<RenameEditAndDeleteMapsFlowTest::ETestActivity>& ExpectedPossibleDependencies = {}
 		)
@@ -73,7 +73,7 @@ namespace UE::ConcertSyncTests::AnalysisTests
 
 		// Delete /Game/Foo > Nearly everything has hard dependency
 		{
-			const FHistoryDeletionRequirements DeleteFooRequirements = AnalyseActivityDeletion({ Activities[_1_NewPackageFoo] }, DependencyGraph);
+			const FHistoryEditionArgs DeleteFooRequirements = AnalyseActivityDependencies({ Activities[_1_NewPackageFoo] }, DependencyGraph);
 			/* _1_NewPackageFoo: is what we're "deleting".
 			 * _5_SavePackageBar: Bar is created as result of a rename but has no dependency to _1_NewPackageFoo.
 			 * All other activities transitively depend on _1_NewPackageFoo (put above log into GraphViz to visualise).
@@ -90,7 +90,7 @@ namespace UE::ConcertSyncTests::AnalysisTests
 		// Delete rename transaction > No dependencies
 			// because the rename basically just changes AActor::ActorLabel
 		{
-			const FHistoryDeletionRequirements DeleteRenameRequirements = AnalyseActivityDeletion({ Activities[_3_RenameActor] }, DependencyGraph);
+			const FHistoryEditionArgs DeleteRenameRequirements = AnalyseActivityDependencies({ Activities[_3_RenameActor] }, DependencyGraph);
 			TestEqual(TEXT("Delete renaming actor: HardDependencies.Num() == 0"), DeleteRenameRequirements.HardDependencies.Num(), 0);
 			TestEqual(TEXT("Delete renaming actor: PossibleDependencies.Num() == 1"), DeleteRenameRequirements.PossibleDependencies.Num(), 1);
 			TestTrue(TEXT("Delete renaming actor: Edit activity may depend on deleted activity"), DeleteRenameRequirements.PossibleDependencies.Contains(Activities[_4_EditActor]));
@@ -98,7 +98,7 @@ namespace UE::ConcertSyncTests::AnalysisTests
 
 		// Delete actor creation > All transactions operating on actor are hard dependencies
 		{
-			const FHistoryDeletionRequirements DeleteCreateActorRequirements = AnalyseActivityDeletion({ Activities[_2_AddActor] }, DependencyGraph);
+			const FHistoryEditionArgs DeleteCreateActorRequirements = AnalyseActivityDependencies({ Activities[_2_AddActor] }, DependencyGraph);
 			TestEqual(TEXT("Delete actor creation: HardDependencies.Num() == 2"), DeleteCreateActorRequirements.HardDependencies.Num(), 2);
 			TestEqual(TEXT("Delete actor creation: PossibleDependencies.Num() == 0"), DeleteCreateActorRequirements.PossibleDependencies.Num(), 0);
 			TestTrue(TEXT("Delete actor creation: Rename depends on created actor"), DeleteCreateActorRequirements.HardDependencies.Contains(Activities[_3_RenameActor]));
@@ -108,7 +108,7 @@ namespace UE::ConcertSyncTests::AnalysisTests
 		// Deleting a rename activity > Rename activity should have a possible dependency to activities that saved the renamed-to package
 			// because MU internally creates a Save Package activity before a rename activity (so secretly rename activities is two related activities)
 		{
-			const FHistoryDeletionRequirements DeleteRenameRequirements = AnalyseActivityDeletion({ Activities[_5_RenameFooToBar] }, DependencyGraph);
+			const FHistoryEditionArgs DeleteRenameRequirements = AnalyseActivityDependencies({ Activities[_5_RenameFooToBar] }, DependencyGraph);
 
 			// Not really what we want to test - but we're doing it for completness
 			TestEqual(TEXT("Delete rename package: HardDependencies.Num() == 4"), DeleteRenameRequirements.HardDependencies.Num(), 4);
@@ -162,7 +162,7 @@ namespace UE::ConcertSyncTests::AnalysisTests
 		DependencyGraph.AddDependency(LeafNodeID, FActivityDependencyEdge(BNodeID, EActivityDependencyReason::EditAfterPreviousPackageEdit, EDependencyStrength::HardDependency));
 		DependencyGraph.AddDependency(BNodeID, FActivityDependencyEdge(RootNodeID, EActivityDependencyReason::EditAfterPreviousPackageEdit, EDependencyStrength::HardDependency));
 
-		const FHistoryDeletionRequirements DeleteFooRequirements = AnalyseActivityDeletion({ RootActivityID }, DependencyGraph);
+		const FHistoryEditionArgs DeleteFooRequirements = AnalyseActivityDependencies({ RootActivityID }, DependencyGraph);
 
 		TestEqual(TEXT("HardDependencies.Num() == 2"), DeleteFooRequirements.HardDependencies.Num(), 2);
 		TestTrue(TEXT("HardDependencies.Contains(B)"), DeleteFooRequirements.HardDependencies.Contains(BActivityID));
