Analyse activities bottom-up

diff --git a/Engine/Plugins/Developer/Concert/ConcertApp/MultiUserServer/Source/MultiUserServer/Private/Widgets/HistoryDeletion/SActivityDependencyView.cpp b/Engine/Plugins/Developer/Concert/ConcertApp/MultiUserServer/Source/MultiUserServer/Private/Widgets/HistoryDeletion/SActivityDependencyView.cpp
--- a/Engine/Plugins/Developer/Concert/ConcertApp/MultiUserServer/Source/MultiUserServer/Private/Widgets/HistoryDeletion/SActivityDependencyView.cpp
+++ b/Engine/Plugins/Developer/Concert/ConcertApp/MultiUserServer/Source/MultiUserServer/Private/Widgets/HistoryDeletion/SActivityDependencyView.cpp
@@ -30,7 +30,7 @@ namespace UE::MultiUserServer
 	}
 }
 
-void SActivityDependencyView::Construct(const FArguments& InArgs, const UE::ConcertSyncCore::FHistoryEditionArgs& DeletionRequirements)
+void SActivityDependencyView::Construct(const FArguments& InArgs, const UE::ConcertSyncCore::FHistoryAnalysisResult& DeletionRequirements)
 {
 	check(InArgs._CreateSessionHistory.IsBound());
 
diff --git a/Engine/Plugins/Developer/Concert/ConcertApp/MultiUserServer/Source/MultiUserServer/Private/Widgets/HistoryDeletion/SActivityDependencyView.h b/Engine/Plugins/Developer/Concert/ConcertApp/MultiUserServer/Source/MultiUserServer/Private/Widgets/HistoryDeletion/SActivityDependencyView.h
--- a/Engine/Plugins/Developer/Concert/ConcertApp/MultiUserServer/Source/MultiUserServer/Private/Widgets/HistoryDeletion/SActivityDependencyView.h
+++ b/Engine/Plugins/Developer/Concert/ConcertApp/MultiUserServer/Source/MultiUserServer/Private/Widgets/HistoryDeletion/SActivityDependencyView.h
@@ -12,7 +12,7 @@ namespace UE
 {
 	namespace ConcertSyncCore
 	{
-		struct FHistoryEditionArgs;
+		struct FHistoryAnalysisResult;
 	}
 }
 
@@ -33,5 +33,5 @@ public:
 		SLATE_ARGUMENT(FText, CustomActivityColumnLabel)
     SLATE_END_ARGS()
 	
-	void Construct(const FArguments& InArgs, const UE::ConcertSyncCore::FHistoryEditionArgs& DeletionRequirements);
+	void Construct(const FArguments& InArgs, const UE::ConcertSyncCore::FHistoryAnalysisResult& DeletionRequirements);
 };
diff --git a/Engine/Plugins/Developer/Concert/ConcertApp/MultiUserServer/Source/MultiUserServer/Private/Widgets/HistoryDeletion/SDeleteActivityDependenciesDialog.cpp b/Engine/Plugins/Developer/Concert/ConcertApp/MultiUserServer/Source/MultiUserServer/Private/Widgets/HistoryDeletion/SDeleteActivityDependenciesDialog.cpp
--- a/Engine/Plugins/Developer/Concert/ConcertApp/MultiUserServer/Source/MultiUserServer/Private/Widgets/HistoryDeletion/SDeleteActivityDependenciesDialog.cpp
+++ b/Engine/Plugins/Developer/Concert/ConcertApp/MultiUserServer/Source/MultiUserServer/Private/Widgets/HistoryDeletion/SDeleteActivityDependenciesDialog.cpp
@@ -12,7 +12,7 @@
 
 #define LOCTEXT_NAMESPACE "UnrealMultiUserUI"
 
-void SDeleteActivityDependenciesDialog::Construct(const FArguments& InArgs, const FGuid& SessionId, const TSharedRef<IConcertSyncServer>& SyncServer, UE::ConcertSyncCore::FHistoryEditionArgs InDeletionRequirements)
+void SDeleteActivityDependenciesDialog::Construct(const FArguments& InArgs, const FGuid& SessionId, const TSharedRef<IConcertSyncServer>& SyncServer, UE::ConcertSyncCore::FHistoryAnalysisResult InDeletionRequirements)
 {
 	DeletionRequirements = MoveTemp(InDeletionRequirements);
 	OnConfirmDeletionFunc = InArgs._OnConfirmDeletion;
@@ -52,7 +52,7 @@ void SDeleteActivityDependenciesDialog::Construct(const FArguments& InArgs, cons
 	);
 }
 
-TSharedRef<SWidget> SDeleteActivityDependenciesDialog::CreateBody(const FGuid& InSessionId, const TSharedRef<IConcertSyncServer>& InSyncServer, UE::ConcertSyncCore::FHistoryEditionArgs InDeletionRequirements)
+TSharedRef<SWidget> SDeleteActivityDependenciesDialog::CreateBody(const FGuid& InSessionId, const TSharedRef<IConcertSyncServer>& InSyncServer, UE::ConcertSyncCore::FHistoryAnalysisResult InDeletionRequirements)
 {
 	for (FActivityID HardDependency : DeletionRequirements.HardDependencies)
 	{
diff --git a/Engine/Plugins/Developer/Concert/ConcertApp/MultiUserServer/Source/MultiUserServer/Private/Widgets/HistoryDeletion/SDeleteActivityDependenciesDialog.h b/Engine/Plugins/Developer/Concert/ConcertApp/MultiUserServer/Source/MultiUserServer/Private/Widgets/HistoryDeletion/SDeleteActivityDependenciesDialog.h
--- a/Engine/Plugins/Developer/Concert/ConcertApp/MultiUserServer/Source/MultiUserServer/Private/Widgets/HistoryDeletion/SDeleteActivityDependenciesDialog.h
+++ b/Engine/Plugins/Developer/Concert/ConcertApp/MultiUserServer/Source/MultiUserServer/Private/Widgets/HistoryDeletion/SDeleteActivityDependenciesDialog.h
@@ -16,7 +16,7 @@ class SDeleteActivityDependenciesDialog : public SCustomDialog
 {
 public:
 
-	DECLARE_DELEGATE_OneParam(FConfirmDeletion, const UE::ConcertSyncCore::FHistoryEditionArgs& /*SelectedRequirements*/)
+	DECLARE_DELEGATE_OneParam(FConfirmDeletion, const UE::ConcertSyncCore::FHistoryAnalysisResult& /*SelectedRequirements*/)
 
 	SLATE_BEGIN_ARGS(SDeleteActivityDependenciesDialog)
 	{}
@@ -27,11 +27,11 @@ public:
 	/**
 	 * @param InDeletionRequirements Specifies which activities must be deleted and which are optional.
 	 */
-	void Construct(const FArguments& InArgs, const FGuid& SessionId, const TSharedRef<IConcertSyncServer>& SyncServer, UE::ConcertSyncCore::FHistoryEditionArgs InDeletionRequirements);
+	void Construct(const FArguments& InArgs, const FGuid& SessionId, const TSharedRef<IConcertSyncServer>& SyncServer, UE::ConcertSyncCore::FHistoryAnalysisResult InDeletionRequirements);
 
 private:
 	
-	UE::ConcertSyncCore::FHistoryEditionArgs DeletionRequirements;
+	UE::ConcertSyncCore::FHistoryAnalysisResult DeletionRequirements;
 	FConfirmDeletion OnConfirmDeletionFunc;
 	
 	/** Maps each activity with a possible dependency to whether it will be deleted or not. */
@@ -40,7 +40,7 @@ private:
 	/** Filters out all activities that are not part of the deletion requirements. */
 	TSharedPtr<FArchivedSessionHistoryController> FilteredSessionHistoryController;
 
-	TSharedRef<SWidget> CreateBody(const FGuid& InSessionId, const TSharedRef<IConcertSyncServer>& InSyncServer, UE::ConcertSyncCore::FHistoryEditionArgs InDeletionRequirements);
+	TSharedRef<SWidget> CreateBody(const FGuid& InSessionId, const TSharedRef<IConcertSyncServer>& InSyncServer, UE::ConcertSyncCore::FHistoryAnalysisResult InDeletionRequirements);
 
 	void OnConfirmPressed();
 };
diff --git a/Engine/Plugins/Developer/Concert/ConcertApp/MultiUserServer/Source/MultiUserServer/Private/Widgets/SessionTabs/Archived/ArchivedConcertSessionTab.cpp b/Engine/Plugins/Developer/Concert/ConcertApp/MultiUserServer/Source/MultiUserServer/Private/Widgets/SessionTabs/Archived/ArchivedConcertSessionTab.cpp
--- a/Engine/Plugins/Developer/Concert/ConcertApp/MultiUserServer/Source/MultiUserServer/Private/Widgets/SessionTabs/Archived/ArchivedConcertSessionTab.cpp
+++ b/Engine/Plugins/Developer/Concert/ConcertApp/MultiUserServer/Source/MultiUserServer/Private/Widgets/SessionTabs/Archived/ArchivedConcertSessionTab.cpp
@@ -67,11 +67,11 @@ void FArchivedConcertSessionTab::OnRequestDeleteActivity(const TSet<TSharedRef<F
 		{
 			return Activity->Activity.ActivityId;
 		});
-		FHistoryEditionArgs DeletionRequirements = AnalyseActivityDependencies(RequestedForDelete, DependencyGraph, true);
+		FHistoryAnalysisResult DeletionRequirements = AnalyseActivityDependencies_TopDown(RequestedForDelete, DependencyGraph, true);
 
 		TWeakPtr<const FArchivedConcertSessionTab> WeakTabThis = SharedThis(this);
 		TSharedRef<SDeleteActivityDependenciesDialog> Dialog = SNew(SDeleteActivityDependenciesDialog, InspectedSessionID, SyncServer, MoveTemp(DeletionRequirements))
-			.OnConfirmDeletion_Lambda([WeakTabThis](const FHistoryEditionArgs& SelectedRequirements)
+			.OnConfirmDeletion_Lambda([WeakTabThis](const FHistoryAnalysisResult& SelectedRequirements)
 			{
 				// Because the dialog is non-modal, the user may have closed the program in the mean time
 				if (const TSharedPtr<const FArchivedConcertSessionTab> PinnedThis = WeakTabThis.Pin())
diff --git a/Engine/Plugins/Developer/Concert/ConcertSync/ConcertSyncCore/Source/ConcertSyncCore/Private/HistoryEdition/HistoryAnalysis.cpp b/Engine/Plugins/Developer/Concert/ConcertSync/ConcertSyncCore/Source/ConcertSyncCore/Private/HistoryEdition/HistoryAnalysis.cpp
deleted file mode 100644
diff --git a/Engine/Plugins/Developer/Concert/ConcertSync/ConcertSyncCore/Source/ConcertSyncCore/Private/HistoryEdition/HistoryBottomUpAnalysis.cpp b/Engine/Plugins/Developer/Concert/ConcertSync/ConcertSyncCore/Source/ConcertSyncCore/Private/HistoryEdition/HistoryBottomUpAnalysis.cpp
new file mode 100644
--- /dev/null
+++ b/Engine/Plugins/Developer/Concert/ConcertSync/ConcertSyncCore/Source/ConcertSyncCore/Private/HistoryEdition/HistoryBottomUpAnalysis.cpp
@@ -0,0 +1,111 @@
+// Copyright Epic Games, Inc. All Rights Reserved.
+
+#include "HistoryEdition/HistoryAnalysis.h"
+
+#include "HistoryEdition/ActivityGraphIDs.h"
+#include "HistoryEdition/DependencyGraphBuilder.h"
+
+#include "Containers/Queue.h"
+
+namespace UE::ConcertSyncCore
+{
+	FHistoryAnalysisResult AnalyseActivityDependencies_BottomUp(const TSet<FActivityID>& ActivitiesToEdit, const FActivityDependencyGraph& DependencyGraph, bool bAddEditedAsHardDependencies)
+	{
+		FHistoryAnalysisResult Result;
+
+		TSet<FActivityNodeID> PossibleDoubleEnqueuingProtection;
+		TQueue<FActivityNodeID> PossibleDependencyActivitiesToAnalyse;
+		FActivityNodeID CurrentActivityID;
+		
+		TSet<FActivityNodeID> HardDoubleEnqueuingProtection;
+		TQueue<FActivityNodeID> HardDependencyActivitiesToAnalyse;
+		for (const FActivityID ActivityToDelete : ActivitiesToEdit)
+		{
+			const TOptional<FActivityNodeID> NodeID = DependencyGraph.FindNodeByActivity(ActivityToDelete);
+			if (ensureMsgf(NodeID, TEXT("Graph does not correspond to ActivitiesToDelete")))
+			{
+				HardDoubleEnqueuingProtection.Add(*NodeID);
+				HardDependencyActivitiesToAnalyse.Enqueue(*NodeID);
+			}
+		}
+		
+		/* We check the hard dependencies first.
+		 * Why? Example:
+		 *
+		 *		R
+		 *	   / \
+		 *	  A   B
+		 *	   \ /
+		 *	    C
+		 *
+		 * The edges C -> A -> R are possible dependencies.
+		 * The edges C -> B -> R are hard dependencies.
+		 *
+		 * Now: delete R.
+		 * We want C to be marked has a hard dependency.
+		 */
+		while (HardDependencyActivitiesToAnalyse.Dequeue(CurrentActivityID))
+		{
+			const FActivityNode& ActivityNode = DependencyGraph.GetNodeById(CurrentActivityID);
+
+			const FActivityID ActivityID = ActivityNode.GetActivityId();
+			if (bAddEditedAsHardDependencies || !ActivitiesToEdit.Contains(ActivityID))
+			{
+				Result.HardDependencies.Add(ActivityID);
+			}
+			
+			for (const FActivityDependencyEdge& Dependency : ActivityNode.GetDependencies())
+			{
+				const FActivityNodeID& DependedOnNodeId = Dependency.GetDependedOnNodeID();
+				if (Dependency.GetDependencyStrength() == EDependencyStrength::HardDependency)
+				{
+					if (!HardDoubleEnqueuingProtection.Contains(DependedOnNodeId))
+					{
+						HardDoubleEnqueuingProtection.Add(DependedOnNodeId);
+						HardDependencyActivitiesToAnalyse.Enqueue(DependedOnNodeId);
+					}
+				}
+				else
+				{
+					check(Dependency.GetDependencyStrength() == EDependencyStrength::PossibleDependency);
+
+					if (!PossibleDoubleEnqueuingProtection.Contains(DependedOnNodeId))
+					{
+						PossibleDoubleEnqueuingProtection.Add(DependedOnNodeId);
+						PossibleDependencyActivitiesToAnalyse.Enqueue(DependedOnNodeId);
+					}
+				}
+			}
+		}
+		
+		// Any possible dependencies that are not also hard dependencies can be added now
+		while (PossibleDependencyActivitiesToAnalyse.Dequeue(CurrentActivityID))
+		{
+			// This would imply a hard dependency - hard dependency takes precedence over possible dependency
+			if (HardDoubleEnqueuingProtection.Contains(CurrentActivityID))
+			{
+				continue;
+			}
+			
+			const FActivityNode& ActivityNode = DependencyGraph.GetNodeById(CurrentActivityID);
+			Result.PossibleDependencies.Add(ActivityNode.GetActivityId());
+			
+			for (const FActivityDependencyEdge& Dependency : ActivityNode.GetDependencies())
+			{
+				const FActivityNodeID& DependedOnNodeId = Dependency.GetDependedOnNodeID();
+				if (Dependency.GetDependencyStrength() == EDependencyStrength::PossibleDependency
+					&& !HardDoubleEnqueuingProtection.Contains(DependedOnNodeId)
+					&& !PossibleDoubleEnqueuingProtection.Contains(DependedOnNodeId))
+				{
+					if (!PossibleDoubleEnqueuingProtection.Contains(DependedOnNodeId))
+					{
+						PossibleDoubleEnqueuingProtection.Add(DependedOnNodeId);
+						PossibleDependencyActivitiesToAnalyse.Enqueue(DependedOnNodeId);
+					}
+				}
+			}
+		}
+		
+		return Result;
+	}
+}
\ No newline at end of file
diff --git a/Engine/Plugins/Developer/Concert/ConcertSync/ConcertSyncCore/Source/ConcertSyncCore/Private/HistoryEdition/HistoryEdition.cpp b/Engine/Plugins/Developer/Concert/ConcertSync/ConcertSyncCore/Source/ConcertSyncCore/Private/HistoryEdition/HistoryEdition.cpp
--- a/Engine/Plugins/Developer/Concert/ConcertSync/ConcertSyncCore/Source/ConcertSyncCore/Private/HistoryEdition/HistoryEdition.cpp
+++ b/Engine/Plugins/Developer/Concert/ConcertSync/ConcertSyncCore/Source/ConcertSyncCore/Private/HistoryEdition/HistoryEdition.cpp
@@ -12,7 +12,7 @@ namespace UE::ConcertSyncCore
 {
 	static FConcertSessionFilter BuildFilterFrom(const TSet<FActivityID>& ToDelete);
 
-	TSet<FActivityID> CombineRequirements(const FHistoryEditionArgs& ToDelete)
+	TSet<FActivityID> CombineRequirements(const FHistoryAnalysisResult& ToDelete)
 	{
 		TSet<FActivityID> Result;
 		for (const FActivityID ActivityID : ToDelete.HardDependencies)
diff --git a/Engine/Plugins/Developer/Concert/ConcertSync/ConcertSyncCore/Source/ConcertSyncCore/Private/HistoryEdition/HistoryTopDownAnalysis.cpp b/Engine/Plugins/Developer/Concert/ConcertSync/ConcertSyncCore/Source/ConcertSyncCore/Private/HistoryEdition/HistoryTopDownAnalysis.cpp
new file mode 100644
--- /dev/null
+++ b/Engine/Plugins/Developer/Concert/ConcertSync/ConcertSyncCore/Source/ConcertSyncCore/Private/HistoryEdition/HistoryTopDownAnalysis.cpp
@@ -0,0 +1,152 @@
+﻿// Copyright Epic Games, Inc. All Rights Reserved.
+
+#include "HistoryEdition/HistoryAnalysis.h"
+
+#include "Algo/Find.h"
+
+#include "HistoryEdition/ActivityGraphIDs.h"
+#include "HistoryEdition/DependencyGraphBuilder.h"
+
+#include "Containers/Queue.h"
+
+namespace UE::ConcertSyncCore
+{
+	FHistoryAnalysisResult AnalyseActivityDependencies_TopDown(const TSet<FActivityID>& ActivitiesToDelete, const FConcertSyncSessionDatabase& Database, bool bAddActivitiesToDelete)
+	{
+		const FActivityDependencyGraph Graph = BuildDependencyGraphFrom(Database);
+		return AnalyseActivityDependencies_TopDown(ActivitiesToDelete, Graph, bAddActivitiesToDelete);
+	}
+
+	/**
+	 * If PossibleRenameActivityNode is a rename activity, then this function adds the activity that created the renamed-to package.
+	 * This implies that PossibleRenameActivityNode is EConcertSyncActivityEventType::Package and has EConcertPackageUpdateType::Renamed.
+	 */
+	static void AddSavePackageActivityAssociatedWithRenamePackageActivity(const FActivityDependencyGraph& DependencyGraph, const FActivityNode& PossibleRenameActivityNode, TSet<FActivityNodeID>& DoubleEnqueuingProtection, TQueue<FActivityNodeID>& ActivitiesToAnalyse);
+	
+	FHistoryAnalysisResult AnalyseActivityDependencies_TopDown(const TSet<FActivityID>& ActivitiesToEdit, const FActivityDependencyGraph& DependencyGraph, bool bAddEditedAsHardDependencies)
+	{
+		FHistoryAnalysisResult Result;
+
+		TSet<FActivityNodeID> PossibleDoubleEnqueuingProtection;
+		TQueue<FActivityNodeID> PossibleDependencyActivitiesToAnalyse;
+		FActivityNodeID CurrentActivityID;
+		
+		TSet<FActivityNodeID> HardDoubleEnqueuingProtection;
+		TQueue<FActivityNodeID> HardDependencyActivitiesToAnalyse;
+		for (const FActivityID ActivityToDelete : ActivitiesToEdit)
+		{
+			const TOptional<FActivityNodeID> NodeID = DependencyGraph.FindNodeByActivity(ActivityToDelete);
+			if (ensureMsgf(NodeID, TEXT("Graph does not correspond to ActivitiesToDelete")))
+			{
+				HardDoubleEnqueuingProtection.Add(*NodeID);
+				HardDependencyActivitiesToAnalyse.Enqueue(*NodeID);
+				AddSavePackageActivityAssociatedWithRenamePackageActivity(DependencyGraph, DependencyGraph.GetNodeById(*NodeID), PossibleDoubleEnqueuingProtection, PossibleDependencyActivitiesToAnalyse);
+			}
+		}
+		
+		/* We check the hard dependencies first.
+		 * Why? Example:
+		 *
+		 *		R
+		 *	   / \
+		 *	  A   B
+		 *	   \ /
+		 *	    C
+		 *
+		 * The edges C -> A -> R are possible dependencies.
+		 * The edges C -> B -> R are hard dependencies.
+		 *
+		 * Now: delete R.
+		 * We want C to be marked has a hard dependency.
+		 */
+		while (HardDependencyActivitiesToAnalyse.Dequeue(CurrentActivityID))
+		{
+			const FActivityNode& ActivityNode = DependencyGraph.GetNodeById(CurrentActivityID);
+
+			const FActivityID ActivityID = ActivityNode.GetActivityId();
+			if (bAddEditedAsHardDependencies || !ActivitiesToEdit.Contains(ActivityID))
+			{
+				Result.HardDependencies.Add(ActivityID);
+			}
+			
+			for (const FActivityNodeID& ChildID : ActivityNode.GetAffectedChildren())
+			{
+				const FActivityNode& ChildNode = DependencyGraph.GetNodeById(ChildID);
+
+				// Performance: The below iterates the edge list twice but usually there will 1 or 2 entries
+				if (ChildNode.DependsOnActivity(ActivityID, DependencyGraph, {}, EDependencyStrength::HardDependency))
+				{
+					if (!HardDoubleEnqueuingProtection.Contains(ChildID))
+					{
+						HardDoubleEnqueuingProtection.Add(ChildID);
+						HardDependencyActivitiesToAnalyse.Enqueue(ChildID);
+						AddSavePackageActivityAssociatedWithRenamePackageActivity(DependencyGraph, ChildNode, PossibleDoubleEnqueuingProtection, PossibleDependencyActivitiesToAnalyse);
+					}
+				}
+				else if (ChildNode.DependsOnActivity(ActivityID, DependencyGraph, {}, EDependencyStrength::PossibleDependency))
+				{
+					if (!HardDoubleEnqueuingProtection.Contains(ChildID) && !PossibleDoubleEnqueuingProtection.Contains(ChildID))
+					{
+						PossibleDoubleEnqueuingProtection.Add(ChildID);
+						PossibleDependencyActivitiesToAnalyse.Enqueue(ChildID);
+						AddSavePackageActivityAssociatedWithRenamePackageActivity(DependencyGraph, ChildNode, PossibleDoubleEnqueuingProtection, PossibleDependencyActivitiesToAnalyse);
+					}
+				}
+			}
+		}
+
+		// Any possible dependencies that are not also hard dependencies can be added now
+		while (PossibleDependencyActivitiesToAnalyse.Dequeue(CurrentActivityID))
+		{
+			// This would imply a hard dependency - hard dependency takes precedence over possible dependency
+			if (HardDoubleEnqueuingProtection.Contains(CurrentActivityID))
+			{
+				continue;
+			}
+			
+			const FActivityNode& ActivityNode = DependencyGraph.GetNodeById(CurrentActivityID);
+			const FActivityID ActivityID = ActivityNode.GetActivityId();
+			Result.PossibleDependencies.Add(ActivityNode.GetActivityId());
+			for (const FActivityNodeID& ChildID : ActivityNode.GetAffectedChildren())
+			{
+				const FActivityNode& ChildNode = DependencyGraph.GetNodeById(ChildID);
+				if (ChildNode.DependsOnActivity(ActivityID, DependencyGraph, {}, EDependencyStrength::PossibleDependency)
+					&& !HardDoubleEnqueuingProtection.Contains(ChildID)
+					&& !PossibleDoubleEnqueuingProtection.Contains(ChildID))
+				{
+					PossibleDoubleEnqueuingProtection.Add(ChildID);
+					PossibleDependencyActivitiesToAnalyse.Enqueue(ChildID);
+					AddSavePackageActivityAssociatedWithRenamePackageActivity(DependencyGraph, ChildNode, PossibleDoubleEnqueuingProtection, PossibleDependencyActivitiesToAnalyse);
+				}
+			}
+		}
+
+		return Result;
+	}
+
+	static void AddSavePackageActivityAssociatedWithRenamePackageActivity(const FActivityDependencyGraph& DependencyGraph, const FActivityNode& PossibleRenameActivityNode, TSet<FActivityNodeID>& DoubleEnqueuingProtection, TQueue<FActivityNodeID>& ActivitiesToAnalyse)
+	{
+		if ((PossibleRenameActivityNode.GetNodeFlags() & EActivityNodeFlags::RenameActivity) == EActivityNodeFlags::None)
+		{
+			return;
+		}
+
+		// Generally rename activities depends on at least one activity and at most on two:
+			// 1. There is always a PossibleDependency with reason PackageCreation: this is the activity that created the renamed-to package; that's the one we want to add here)
+			// 2. There is sometimes a HardDependency with reason PackageCreation in case the renamed package was created within the same session
+		const FActivityDependencyEdge* PackageCreationDependency = Algo::FindByPredicate(PossibleRenameActivityNode.GetDependencies(),
+			[](const FActivityDependencyEdge& Dependency)
+			{
+				const bool bIsAssociatedCreationActivity = Dependency.GetDependencyReason() == EActivityDependencyReason::PackageCreation && Dependency.GetDependencyStrength() == EDependencyStrength::PossibleDependency;
+				return bIsAssociatedCreationActivity;
+			});
+		if (ensureMsgf(PackageCreationDependency, TEXT("Rename activities should have a possible dependency to the activity that creatd the renamed-to package"))
+			&& !DoubleEnqueuingProtection.Contains(PackageCreationDependency->GetDependedOnNodeID()))
+		{
+			DoubleEnqueuingProtection.Add(PackageCreationDependency->GetDependedOnNodeID());
+			ActivitiesToAnalyse.Enqueue(PackageCreationDependency->GetDependedOnNodeID());
+		}
+	}
+}
+
+
diff --git a/Engine/Plugins/Developer/Concert/ConcertSync/ConcertSyncCore/Source/ConcertSyncCore/Public/HistoryEdition/HistoryAnalysis.h b/Engine/Plugins/Developer/Concert/ConcertSync/ConcertSyncCore/Source/ConcertSyncCore/Public/HistoryEdition/HistoryAnalysis.h
--- a/Engine/Plugins/Developer/Concert/ConcertSync/ConcertSyncCore/Source/ConcertSyncCore/Public/HistoryEdition/HistoryAnalysis.h
+++ b/Engine/Plugins/Developer/Concert/ConcertSync/ConcertSyncCore/Source/ConcertSyncCore/Public/HistoryEdition/HistoryAnalysis.h
@@ -11,7 +11,7 @@ class FConcertSyncSessionDatabase;
 namespace UE::ConcertSyncCore
 {
 	/** Describes the activities that must be considered when deleting an activity */
-	struct CONCERTSYNCCORE_API FHistoryEditionArgs
+	struct CONCERTSYNCCORE_API FHistoryAnalysisResult
 	{
 		/** The activities that must be removed. */
 		TSet<FActivityID> HardDependencies;
@@ -21,22 +21,37 @@ namespace UE::ConcertSyncCore
 		 */
 		TSet<FActivityID> PossibleDependencies;
 
-		FHistoryEditionArgs(TSet<FActivityID> HardDependencies = {}, TSet<FActivityID> PossibleDependencies = {})
-			: HardDependencies(MoveTemp(HardDependencies)),
-			  PossibleDependencies(MoveTemp(PossibleDependencies))
+		FHistoryAnalysisResult(TSet<FActivityID> HardDependencies = {}, TSet<FActivityID> PossibleDependencies = {})
+			: HardDependencies(MoveTemp(HardDependencies))
+			, PossibleDependencies(MoveTemp(PossibleDependencies))
 		{}
 	};
 
 	/** Utility function for one-off operations: just computes the dependency graph before calling AnalyseActivityDeletion. */
-	CONCERTSYNCCORE_API FHistoryEditionArgs AnalyseActivityDependencies(const TSet<FActivityID>& ActivitiesToDelete, const FConcertSyncSessionDatabase& Database, bool bAddActivitiesToDelete = false);
+	CONCERTSYNCCORE_API FHistoryAnalysisResult AnalyseActivityDependencies_TopDown(const TSet<FActivityID>& ActivitiesToDelete, const FConcertSyncSessionDatabase& Database, bool bAddActivitiesToDelete = false);
 
 	/**
 	 * Given a set of activities to be edited (e.g. deleted or muted), returns which activities 1. must be and 2. may want to be considered in addition.
-	 * @param ActivitiesToEdit The activities that should be removed
+	 * This algorithm proceeds top-down meaning it will start at ActivitiesToEdit and recursively process child nodes that depend on it.
+	 * This approach is good for finding activities that depends on a set of activities, e.g. for deleting or muting ActivitiesToEdit.
+	 * 
+	 * @param ActivitiesToEdit The activities whose children to search
 	 * @param DependencyGraph The graph encoding the activity dependencies
 	 * @param bAddEditedAsHardDependencies Whether to add ActivitiesToEdit to the result's HardDependencies
 	 * @return The activities to consider as well if ActivitiesToEdit are edited; ActivitiesToEdit is included in HardDependencies if bAddActivitiesToDelete == true.
 	 */
-	CONCERTSYNCCORE_API FHistoryEditionArgs AnalyseActivityDependencies(const TSet<FActivityID>& ActivitiesToEdit, const FActivityDependencyGraph& DependencyGraph, bool bAddEditedAsHardDependencies = false);
+	CONCERTSYNCCORE_API FHistoryAnalysisResult AnalyseActivityDependencies_TopDown(const TSet<FActivityID>& ActivitiesToEdit, const FActivityDependencyGraph& DependencyGraph, bool bAddEditedAsHardDependencies = false);
+
+	/**
+	 * Given a set of activities to be edited (e.g. to be unmuted), returns which activities 1. must be and 2. may want to be considered in addition.
+	 * This algorithm proceeds bottom-up meaning it will start at ActivitiesToEdit and recursively process parent nodes it depends on.
+	 * This approach is good for finding the activities that are depended on by set of activities, e.g. unmuting ActivitiesToEdit.
+	 * 
+	 * @param ActivitiesToEdit The activities whose parents to search
+	 * @param DependencyGraph The graph encoding the activity dependencies
+	 * @param bAddEditedAsHardDependencies Whether to add ActivitiesToEdit to the result's HardDependencies
+	 * @return The activities to consider as well if ActivitiesToEdit are edited; ActivitiesToEdit is included in HardDependencies if bAddActivitiesToDelete == true.
+	 */
+	CONCERTSYNCCORE_API FHistoryAnalysisResult AnalyseActivityDependencies_BottomUp(const TSet<FActivityID>& ActivitiesToEdit, const FActivityDependencyGraph& DependencyGraph, bool bAddEditedAsHardDependencies = false);
 }
 
diff --git a/Engine/Plugins/Developer/Concert/ConcertSync/ConcertSyncCore/Source/ConcertSyncCore/Public/HistoryEdition/HistoryEdition.h b/Engine/Plugins/Developer/Concert/ConcertSync/ConcertSyncCore/Source/ConcertSyncCore/Public/HistoryEdition/HistoryEdition.h
--- a/Engine/Plugins/Developer/Concert/ConcertSync/ConcertSyncCore/Source/ConcertSyncCore/Public/HistoryEdition/HistoryEdition.h
+++ b/Engine/Plugins/Developer/Concert/ConcertSync/ConcertSyncCore/Source/ConcertSyncCore/Public/HistoryEdition/HistoryEdition.h
@@ -26,7 +26,7 @@ namespace UE::ConcertSyncCore
 	};
 
 	/** Utility functions that converts FHistoryDeletionRequirements into a single TSet. */
-	CONCERTSYNCCORE_API TSet<FActivityID> CombineRequirements(const FHistoryEditionArgs& ToDelete);
+	CONCERTSYNCCORE_API TSet<FActivityID> CombineRequirements(const FHistoryAnalysisResult& ToDelete);
 
 	/**
 	 * Deletes the given sessions in ToDelete from ArchivedSessionDatabase.
diff --git a/Engine/Plugins/Developer/Concert/ConcertSync/ConcertSyncTest/Source/ConcertSyncTest/Private/HistoryEdition/HistoryAnalysisTests.cpp b/Engine/Plugins/Developer/Concert/ConcertSync/ConcertSyncTest/Source/ConcertSyncTest/Private/HistoryEdition/HistoryAnalysisTests.cpp
deleted file mode 100644
diff --git a/Engine/Plugins/Developer/Concert/ConcertSync/ConcertSyncTest/Source/ConcertSyncTest/Private/HistoryEdition/HistoryBottomUpAnalysisTests.cpp b/Engine/Plugins/Developer/Concert/ConcertSync/ConcertSyncTest/Source/ConcertSyncTest/Private/HistoryEdition/HistoryBottomUpAnalysisTests.cpp
new file mode 100644
--- /dev/null
+++ b/Engine/Plugins/Developer/Concert/ConcertSync/ConcertSyncTest/Source/ConcertSyncTest/Private/HistoryEdition/HistoryBottomUpAnalysisTests.cpp
@@ -0,0 +1,106 @@
+// Copyright Epic Games, Inc. All Rights Reserved.
+
+#include "ConcertLogGlobal.h"
+
+#include "ConcertSyncSessionDatabase.h"
+#include "HistoryEdition/HistoryAnalysis.h"
+#include "RenameEditAndDeleteMapsFlow.h"
+#include "ScopedSessionDatabase.h"
+
+#include "Algo/AllOf.h"
+
+#include "Core/Tests/Containers/TestUtils.h"
+
+#include "HistoryEdition/DebugDependencyGraph.h"
+#include "HistoryEdition/DependencyGraphBuilder.h"
+
+#include "Misc/AutomationTest.h"
+
+namespace UE::ConcertSyncTests::AnalysisTests
+{
+	/**
+	 * Runs through a couple of use cases of unmuting an activity.
+	 */
+	IMPLEMENT_SIMPLE_AUTOMATION_TEST(FBottomUpUseCase, "Concert.History.Analysis.BottomUpUseCase", EAutomationTestFlags::ApplicationContextMask | EAutomationTestFlags::EngineFilter)
+	bool FBottomUpUseCase::RunTest(const FString& Parameters)
+	{
+		using namespace RenameEditAndDeleteMapsFlowTest;
+		using namespace ConcertSyncCore;
+		
+		FScopedSessionDatabase SessionDatabase(*this);
+		const TTestActivityArray<FActivityID> Activities = CreateActivityHistory(SessionDatabase, SessionDatabase.GetEndpoint());
+		const FActivityDependencyGraph DependencyGraph = BuildDependencyGraphFrom(SessionDatabase);
+		UE_LOG(LogConcert, Log, TEXT("%s tested graph in Graphviz format:\n\n%s"), *GetTestFullName(), *ConcertSyncCore::Graphviz::ExportToGraphviz(DependencyGraph, SessionDatabase));
+
+		// I highly recommend visualising the above graph to understand the test better
+		// Unmute an edit operation that took place after an actor rename
+		{
+			const FHistoryAnalysisResult UnmuteEditAfterRename = AnalyseActivityDependencies_BottomUp({ Activities[_4_EditActor] }, DependencyGraph);
+			TestEqual(TEXT("UnmuteEditAfterRename.HardDependencies.Num() == 2"), UnmuteEditAfterRename.HardDependencies.Num(), 2);
+			TestEqual(TEXT("UnmuteEditAfterRename.PossibleDependencies.Num() == 1"), UnmuteEditAfterRename.PossibleDependencies.Num(), 1);
+			TestTrue(TEXT("Editing actor depends on creating actor"), UnmuteEditAfterRename.HardDependencies.Contains(Activities[_2_AddActor]));
+			TestTrue(TEXT("Creating actor depends on creating map"), UnmuteEditAfterRename.HardDependencies.Contains(Activities[_1_NewPackageFoo]));
+			TestTrue(TEXT("Editting possibly depends on renaming"), UnmuteEditAfterRename.PossibleDependencies.Contains(Activities[_3_RenameActor]));
+		}
+
+		// Unmute package rename
+		{
+			const FHistoryAnalysisResult UnmutePackageRename = AnalyseActivityDependencies_BottomUp({ Activities[_5_RenameFooToBar] }, DependencyGraph);
+			TestEqual(TEXT("UnmuteEditAfterRename.HardDependencies.Num() == 1"), UnmutePackageRename.HardDependencies.Num(), 1);
+			TestEqual(TEXT("UnmuteEditAfterRename.PossibleDependencies.Num() == 1"), UnmutePackageRename.PossibleDependencies.Num(), 1);
+			TestTrue(TEXT("Renaming package depends on having create the package"), UnmutePackageRename.HardDependencies.Contains(Activities[_1_NewPackageFoo]));
+			TestTrue(TEXT("Renaming package possibly depends on saving package"), UnmutePackageRename.PossibleDependencies.Contains(Activities[_5_SavePackageBar]));
+		}
+		
+		return true;
+	}
+
+	/**
+	 * Suppose:
+	 *
+	 *		R
+	 *	   / \
+	 *	  A   B
+	 *	   \ /
+	 *	    L
+	 *
+	 * The edges L -> A -> R are possible dependencies.
+	 * The edges L -> B -> R are hard dependencies.
+	 *
+	 * The test: unmute L.
+	 * We want R to be marked has a hard dependency.
+	 */
+	IMPLEMENT_SIMPLE_AUTOMATION_TEST(FBottomUpDiamond, "Concert.History.Analysis.BottomUpDiamond", EAutomationTestFlags::ApplicationContextMask | EAutomationTestFlags::EngineFilter)
+	bool FBottomUpDiamond::RunTest(const FString& Parameters)
+	{
+		using namespace ConcertSyncCore;
+
+		constexpr FActivityID RootActivityID = 1;
+		constexpr FActivityID AActivityID = 2;
+		constexpr FActivityID BActivityID = 3;
+		constexpr FActivityID LeafActivityID  = 4;
+		
+		FActivityDependencyGraph DependencyGraph;
+		const FActivityNodeID RootNodeID = DependencyGraph.AddActivity(RootActivityID);
+		const FActivityNodeID ANodeID = DependencyGraph.AddActivity(AActivityID);
+		const FActivityNodeID BNodeID = DependencyGraph.AddActivity(BActivityID);
+		const FActivityNodeID LeafNodeID = DependencyGraph.AddActivity(LeafActivityID);
+
+		// Add the weak dependency first so the algorithm finds it first when iterating
+		DependencyGraph.AddDependency(LeafNodeID, FActivityDependencyEdge(ANodeID, EActivityDependencyReason::EditAfterPreviousPackageEdit, EDependencyStrength::PossibleDependency));
+		DependencyGraph.AddDependency(ANodeID, FActivityDependencyEdge(RootNodeID, EActivityDependencyReason::EditAfterPreviousPackageEdit, EDependencyStrength::PossibleDependency));
+		DependencyGraph.AddDependency(LeafNodeID, FActivityDependencyEdge(BNodeID, EActivityDependencyReason::EditAfterPreviousPackageEdit, EDependencyStrength::HardDependency));
+		DependencyGraph.AddDependency(BNodeID, FActivityDependencyEdge(RootNodeID, EActivityDependencyReason::EditAfterPreviousPackageEdit, EDependencyStrength::HardDependency));
+
+		const FHistoryAnalysisResult DeleteFooRequirements = AnalyseActivityDependencies_BottomUp({ LeafActivityID }, DependencyGraph);
+
+		TestEqual(TEXT("HardDependencies.Num() == 2"), DeleteFooRequirements.HardDependencies.Num(), 2);
+		TestTrue(TEXT("HardDependencies.Contains(B)"), DeleteFooRequirements.HardDependencies.Contains(BActivityID));
+		TestTrue(TEXT("HardDependencies.Contains(L)"), DeleteFooRequirements.HardDependencies.Contains(RootActivityID));
+		
+		TestEqual(TEXT("PossibleDependencies.Num() == 1"), DeleteFooRequirements.PossibleDependencies.Num(), 1);
+		TestTrue(TEXT("PossibleDependencies.Contains(A)"), DeleteFooRequirements.PossibleDependencies.Contains(AActivityID));
+
+		return true;
+	}
+}
\ No newline at end of file
diff --git a/Engine/Plugins/Developer/Concert/ConcertSync/ConcertSyncTest/Source/ConcertSyncTest/Private/HistoryEdition/HistoryTopDownAnalysisTests.cpp b/Engine/Plugins/Developer/Concert/ConcertSync/ConcertSyncTest/Source/ConcertSyncTest/Private/HistoryEdition/HistoryTopDownAnalysisTests.cpp
new file mode 100644
--- /dev/null
+++ b/Engine/Plugins/Developer/Concert/ConcertSync/ConcertSyncTest/Source/ConcertSyncTest/Private/HistoryEdition/HistoryTopDownAnalysisTests.cpp
@@ -0,0 +1,176 @@
+// Copyright Epic Games, Inc. All Rights Reserved.
+
+#include "ConcertLogGlobal.h"
+
+#include "ConcertSyncSessionDatabase.h"
+#include "HistoryEdition/HistoryAnalysis.h"
+#include "RenameEditAndDeleteMapsFlow.h"
+#include "ScopedSessionDatabase.h"
+
+#include "Algo/AllOf.h"
+
+#include "Core/Tests/Containers/TestUtils.h"
+
+#include "HistoryEdition/DebugDependencyGraph.h"
+#include "HistoryEdition/DependencyGraphBuilder.h"
+
+#include "Misc/AutomationTest.h"
+
+namespace UE::ConcertSyncTests
+{
+	bool ValidateRequirements(
+		const FString& TestBaseName,
+		FAutomationTestBase& Test,
+		const RenameEditAndDeleteMapsFlowTest::TTestActivityArray<FActivityID> Activities,
+		const ConcertSyncCore::FHistoryAnalysisResult& ToValidate = {},
+		const TSet<RenameEditAndDeleteMapsFlowTest::ETestActivity>& ExpectedHardDependencies = {},
+		const TSet<RenameEditAndDeleteMapsFlowTest::ETestActivity>& ExpectedPossibleDependencies = {}
+		)
+	{
+		using namespace RenameEditAndDeleteMapsFlowTest;
+
+		const bool bHardDependenciesAreCorrect = Algo::AllOf(ExpectedHardDependencies, [&TestBaseName, &Test, &Activities, &ToValidate](ETestActivity Activity)
+		{
+			const bool bContained = ToValidate.HardDependencies.Contains(Activities[Activity]);
+			Test.TestTrue(FString::Printf(TEXT("%s: %s is a hard dependency"), *TestBaseName, *LexToString(Activity)), bContained);
+			return bContained;
+		});
+		Test.TestTrue(FString::Printf(TEXT("%s: Hard dependencies are correct"), *TestBaseName), bHardDependenciesAreCorrect);
+		
+		const bool bPossibleDependenciesAreCorrect = Algo::AllOf(ExpectedPossibleDependencies, [&TestBaseName, &Test, &Activities, &ToValidate](ETestActivity Activity)
+		{
+			const bool bContained = ToValidate.PossibleDependencies.Contains(Activities[Activity]);
+			Test.TestTrue(FString::Printf(TEXT("%s: %s is a possible dependency"), *TestBaseName, *LexToString(Activity)), bContained);
+			return bContained;
+		});
+		Test.TestTrue(FString::Printf(TEXT("%s: Possible dependencies are correct"), *TestBaseName), bPossibleDependenciesAreCorrect);
+		
+		const TSet<ETestActivity> ExpectedExcludedActivities = AllActivities().Difference(ExpectedHardDependencies.Union(ExpectedPossibleDependencies));
+		const bool bAllOtherActivitiesExcluded = Algo::AllOf(ExpectedExcludedActivities, [&TestBaseName, &Test, &Activities, &ToValidate](ETestActivity Activity)
+		{
+			const bool bNotContained = !ToValidate.HardDependencies.Contains(Activities[Activity]) && !ToValidate.PossibleDependencies.Contains(Activities[Activity]);
+			Test.TestTrue(FString::Printf(TEXT("%s: %s is no dependency"), *TestBaseName, *LexToString(Activity)), bNotContained);
+			return bNotContained;
+		});
+		Test.TestTrue(FString::Printf(TEXT("%s: No unexpected dependencies"), *TestBaseName), bAllOtherActivitiesExcluded);
+
+		return bHardDependenciesAreCorrect && bPossibleDependenciesAreCorrect && bAllOtherActivitiesExcluded;
+	}
+}
+
+namespace UE::ConcertSyncTests::AnalysisTests
+{
+	IMPLEMENT_SIMPLE_AUTOMATION_TEST(FTopDownUseCase, "Concert.History.Analysis.TopDownUseCase", EAutomationTestFlags::ApplicationContextMask | EAutomationTestFlags::EngineFilter)
+	bool FTopDownUseCase::RunTest(const FString& Parameters)
+	{
+		using namespace RenameEditAndDeleteMapsFlowTest;
+		using namespace ConcertSyncCore;
+		
+		FScopedSessionDatabase SessionDatabase(*this);
+		const TTestActivityArray<FActivityID> Activities = CreateActivityHistory(SessionDatabase, SessionDatabase.GetEndpoint());
+		const FActivityDependencyGraph DependencyGraph = BuildDependencyGraphFrom(SessionDatabase);
+		UE_LOG(LogConcert, Log, TEXT("%s tested graph in Graphviz format:\n\n%s"), *GetTestFullName(), *ConcertSyncCore::Graphviz::ExportToGraphviz(DependencyGraph, SessionDatabase));
+
+		// Delete /Game/Foo > Nearly everything has hard dependency
+		{
+			const FHistoryAnalysisResult DeleteFooRequirements = AnalyseActivityDependencies_TopDown({ Activities[_1_NewPackageFoo] }, DependencyGraph);
+			/* _1_NewPackageFoo: is what we're "deleting".
+			 * _5_SavePackageBar: Bar is created as result of a rename but has no dependency to _1_NewPackageFoo.
+			 * All other activities transitively depend on _1_NewPackageFoo (put above log into GraphViz to visualise).
+			 * 
+			 * Note: The transaction activities (_3_EditActor, _4_EditActor) have possible dependencies BUT they do have hard dependencies to _2_AddActor. This is why they must be in HardDependencies, too.
+			 */
+			const TSet<ETestActivity> ExcludedHardActivities { _1_NewPackageFoo, _5_SavePackageBar };
+			const TSet<ETestActivity> HardDependencies = AllActivities().Difference(ExcludedHardActivities);
+			const TSet<ETestActivity> PossibleDependencies { _5_SavePackageBar };
+			const bool bDeleteAllCorrect = ValidateRequirements(TEXT("Delete /Game/Foo"), *this, Activities, DeleteFooRequirements, HardDependencies, PossibleDependencies);
+			TestTrue(TEXT("Delete /Game/Foo is correct"), bDeleteAllCorrect);
+		}
+
+		// Delete rename transaction > No dependencies
+			// because the rename basically just changes AActor::ActorLabel
+		{
+			const FHistoryAnalysisResult DeleteRenameRequirements = AnalyseActivityDependencies_TopDown({ Activities[_3_RenameActor] }, DependencyGraph);
+			TestEqual(TEXT("Delete renaming actor: HardDependencies.Num() == 0"), DeleteRenameRequirements.HardDependencies.Num(), 0);
+			TestEqual(TEXT("Delete renaming actor: PossibleDependencies.Num() == 1"), DeleteRenameRequirements.PossibleDependencies.Num(), 1);
+			TestTrue(TEXT("Delete renaming actor: Edit activity may depend on deleted activity"), DeleteRenameRequirements.PossibleDependencies.Contains(Activities[_4_EditActor]));
+		}
+
+		// Delete actor creation > All transactions operating on actor are hard dependencies
+		{
+			const FHistoryAnalysisResult DeleteCreateActorRequirements = AnalyseActivityDependencies_TopDown({ Activities[_2_AddActor] }, DependencyGraph);
+			TestEqual(TEXT("Delete actor creation: HardDependencies.Num() == 2"), DeleteCreateActorRequirements.HardDependencies.Num(), 2);
+			TestEqual(TEXT("Delete actor creation: PossibleDependencies.Num() == 0"), DeleteCreateActorRequirements.PossibleDependencies.Num(), 0);
+			TestTrue(TEXT("Delete actor creation: Rename depends on created actor"), DeleteCreateActorRequirements.HardDependencies.Contains(Activities[_3_RenameActor]));
+			TestTrue(TEXT("Delete actor creation: Edit depends on created actor"), DeleteCreateActorRequirements.HardDependencies.Contains(Activities[_4_EditActor]));
+		}
+
+		// Deleting a rename activity > Rename activity should have a possible dependency to activities that saved the renamed-to package
+			// because MU internally creates a Save Package activity before a rename activity (so secretly rename activities is two related activities)
+		{
+			const FHistoryAnalysisResult DeleteRenameRequirements = AnalyseActivityDependencies_TopDown({ Activities[_5_RenameFooToBar] }, DependencyGraph);
+
+			// Not really what we want to test - but we're doing it for completness
+			TestEqual(TEXT("Delete rename package: HardDependencies.Num() == 4"), DeleteRenameRequirements.HardDependencies.Num(), 4);
+			TestTrue(TEXT("Delete rename package: Rename has hard dependency creating renamed package"), DeleteRenameRequirements.HardDependencies.Contains(Activities[_6_EditActor]));
+			TestTrue(TEXT("Delete rename package: Rename has hard dependency to editing renamed package (1)"), DeleteRenameRequirements.HardDependencies.Contains(Activities[_7_DeleteBar]));
+			TestTrue(TEXT("Delete rename package: Rename has hard dependency to editing renamed package (2)"), DeleteRenameRequirements.HardDependencies.Contains(Activities[_8_NewPackageFoo]));
+			TestTrue(TEXT("Delete rename package: Rename has hard dependency to editing renamed package (3)"), DeleteRenameRequirements.HardDependencies.Contains(Activities[_8_SavePackageFoo]));
+
+			// Now what we really wanted to test
+			TestEqual(TEXT("Delete rename package: PossibleDependencies.Num() == 1"), DeleteRenameRequirements.PossibleDependencies.Num(), 1);
+			TestTrue(TEXT("Delete rename package: Rename has soft dependency to saving renamed-to package"), DeleteRenameRequirements.PossibleDependencies.Contains(Activities[_5_SavePackageBar]));
+		}
+
+		return true;
+	}
+
+	/**
+	 * Suppose:
+	 *
+	 *		R
+	 *	   / \
+	 *	  A   B
+	 *	   \ /
+	 *	    L
+	 *
+	 * The edges L -> A -> R are possible dependencies.
+	 * The edges L -> B -> R are hard dependencies.
+	 *
+	 * The test: delete R.
+	 * We want L to be marked has a hard dependency.
+	 */
+	IMPLEMENT_SIMPLE_AUTOMATION_TEST(FTopDownDiamond, "Concert.History.Analysis.TopDownDiamond", EAutomationTestFlags::ApplicationContextMask | EAutomationTestFlags::EngineFilter)
+	bool FTopDownDiamond::RunTest(const FString& Parameters)
+	{
+		using namespace ConcertSyncCore;
+
+		constexpr FActivityID RootActivityID = 1;
+		constexpr FActivityID AActivityID = 2;
+		constexpr FActivityID BActivityID = 3;
+		constexpr FActivityID LeafActivityID  = 4;
+		
+		FActivityDependencyGraph DependencyGraph;
+		const FActivityNodeID RootNodeID = DependencyGraph.AddActivity(RootActivityID);
+		const FActivityNodeID ANodeID = DependencyGraph.AddActivity(AActivityID);
+		const FActivityNodeID BNodeID = DependencyGraph.AddActivity(BActivityID);
+		const FActivityNodeID LeafNodeID = DependencyGraph.AddActivity(LeafActivityID);
+
+		// Add the weak dependency first so the algorithm finds it first when iterating
+		DependencyGraph.AddDependency(LeafNodeID, FActivityDependencyEdge(ANodeID, EActivityDependencyReason::EditAfterPreviousPackageEdit, EDependencyStrength::PossibleDependency));
+		DependencyGraph.AddDependency(ANodeID, FActivityDependencyEdge(RootNodeID, EActivityDependencyReason::EditAfterPreviousPackageEdit, EDependencyStrength::PossibleDependency));
+		DependencyGraph.AddDependency(LeafNodeID, FActivityDependencyEdge(BNodeID, EActivityDependencyReason::EditAfterPreviousPackageEdit, EDependencyStrength::HardDependency));
+		DependencyGraph.AddDependency(BNodeID, FActivityDependencyEdge(RootNodeID, EActivityDependencyReason::EditAfterPreviousPackageEdit, EDependencyStrength::HardDependency));
+
+		const FHistoryAnalysisResult DeleteFooRequirements = AnalyseActivityDependencies_TopDown({ RootActivityID }, DependencyGraph);
+
+		TestEqual(TEXT("HardDependencies.Num() == 2"), DeleteFooRequirements.HardDependencies.Num(), 2);
+		TestTrue(TEXT("HardDependencies.Contains(B)"), DeleteFooRequirements.HardDependencies.Contains(BActivityID));
+		TestTrue(TEXT("HardDependencies.Contains(L)"), DeleteFooRequirements.HardDependencies.Contains(LeafActivityID));
+		
+		TestEqual(TEXT("PossibleDependencies.Num() == 1"), DeleteFooRequirements.PossibleDependencies.Num(), 1);
+		TestTrue(TEXT("PossibleDependencies.Contains(A)"), DeleteFooRequirements.PossibleDependencies.Contains(AActivityID));
+
+		return true;
+	}
+}
\ No newline at end of file
