Make saving dependend on transaction activities. This fixes deleting activities not working correctly.

diff --git a/Engine/Plugins/Developer/Concert/ConcertSync/ConcertSyncCore/Source/ConcertSyncCore/Private/HistoryEdition/ActivityDependencyEdge.cpp b/Engine/Plugins/Developer/Concert/ConcertSync/ConcertSyncCore/Source/ConcertSyncCore/Private/HistoryEdition/ActivityDependencyEdge.cpp
--- a/Engine/Plugins/Developer/Concert/ConcertSync/ConcertSyncCore/Source/ConcertSyncCore/Private/HistoryEdition/ActivityDependencyEdge.cpp
+++ b/Engine/Plugins/Developer/Concert/ConcertSync/ConcertSyncCore/Source/ConcertSyncCore/Private/HistoryEdition/ActivityDependencyEdge.cpp
@@ -9,11 +9,13 @@ FString UE::ConcertSyncCore::LexToString(EActivityDependencyReason Reason)
 		case EActivityDependencyReason::PackageCreation: return TEXT("PackageCreation");
 		case EActivityDependencyReason::PackageRemoval: return TEXT("PackageRemoval");
 		case EActivityDependencyReason::PackageRename: return TEXT("PackageRename");
+		case EActivityDependencyReason::PackageEdited: return TEXT("PackageEdited");
 		case EActivityDependencyReason::EditAfterPreviousPackageEdit: return TEXT("EditPossiblyDependsOnPackage");
 		case EActivityDependencyReason::SubobjectCreation: return TEXT("SubobjectCreation");
 		case EActivityDependencyReason::SubobjectRemoval: return TEXT("SubobjectRemoval");
 		default:
 			checkNoEntry();
+			static_assert(static_cast<int32>(EActivityDependencyReason::Count) == 7, "Update this location when you change EActivityDependencyReason");
 			return TEXT("Invalid");
 	}
 }
diff --git a/Engine/Plugins/Developer/Concert/ConcertSync/ConcertSyncCore/Source/ConcertSyncCore/Private/HistoryEdition/DependencyGraphBuilder.cpp b/Engine/Plugins/Developer/Concert/ConcertSync/ConcertSyncCore/Source/ConcertSyncCore/Private/HistoryEdition/DependencyGraphBuilder.cpp
--- a/Engine/Plugins/Developer/Concert/ConcertSync/ConcertSyncCore/Source/ConcertSyncCore/Private/HistoryEdition/DependencyGraphBuilder.cpp
+++ b/Engine/Plugins/Developer/Concert/ConcertSync/ConcertSyncCore/Source/ConcertSyncCore/Private/HistoryEdition/DependencyGraphBuilder.cpp
@@ -14,8 +14,13 @@ namespace UE::ConcertSyncCore
 		TMap<FName, FActivityID> AddedPackages;
 		TMap<FName, FActivityID> SavedPackages;
 		TMap<FName, FActivityID> RemovedPackages;
-		TMap<FName, FActivityID> RenamedPackages;
+		/** Value: The last activity that renamed a package to the key. */
+		TMap<FName, FActivityID> RenamedToPackages;
+		/** Value: The last activity that renamed a package from the key to something else. */
+		TMap<FName, FActivityID> RenamedFromPackages;
 		TMap<FName, FActivityID> ModifiedPackages;
+		/** The first activity that modified the given package after a revision of the package has been created. */
+		TMap<FName, TArray<FActivityID>> TransactionsAffectingPackageSinceLastPackageActivity;
 
 		enum class ESubobjectState
 		{
@@ -56,10 +61,14 @@ namespace UE::ConcertSyncCore
 		void DiscoverRenamedPackageDependencies(const FActivityNodeID NodeID, const FConcertSyncActivity& Activity, const FConcertSyncPackageEventMetaData& EventData);
 		void DiscoverDeletedPackageDependencies(const FActivityNodeID NodeID, const FConcertSyncActivity& Activity, const FConcertSyncPackageEventMetaData& EventData);
 
+		void DiscoverPackageEditedDependencies(const FActivityNodeID DependingNodeID, FName ModifiedPackage);
+
 		/** Util function for declarative style dependency statement in C++ */
 		enum class EPackageAddDependencyCondition
 		{
+			/** The activity will be added unconditionally */
 			Always,
+			/** From the array of activities, only the one with the highest activity ID will be added. */
 			OnlyLatestActivity
 		};
 		using FPackageActivityItem = TTuple<const FActivityID* /* ActivityID */, EActivityDependencyReason /* Reason */, EDependencyStrength /* DependencyStrength*/, EPackageAddDependencyCondition /* AddCondition */>;
@@ -151,7 +160,16 @@ namespace UE::ConcertSyncCore
 		for (const FName ModifiedPackage : EventData.Transaction.ModifiedPackages)
 		{
 			const FActivityID* LastModifyingActivity = PackageTracker.ModifiedPackages.Find(ModifiedPackage);
-			if (LastModifyingActivity)
+			const FActivityID* LastAddedActivity = PackageTracker.AddedPackages.Find(ModifiedPackage);
+			const FActivityID* LastRenamingActivity = PackageTracker.RenamedToPackages.Find(ModifiedPackage);
+
+			// We might depend on an activity that modified the same package but only if it is still the same package, i.e.
+			// 1. The package was not added again between the previous transaction and the current (would imply delete or rename inbetween)
+			// 2. The package was not renamed between the previous transaction and the current
+			const bool bCanConsiderLastModifyingActivity = LastModifyingActivity
+				&& (!LastAddedActivity || *LastAddedActivity < *LastModifyingActivity)
+				&& (!LastRenamingActivity || *LastRenamingActivity < *LastModifyingActivity);
+			if (bCanConsiderLastModifyingActivity)
 			{
 				// If there is already a hard dependency to a previous activity, then we do not "possibly" depend on it: we definitely depend on it (already).
 				if (!CreatedDependencies.Contains(*LastModifyingActivity))
@@ -160,14 +178,14 @@ namespace UE::ConcertSyncCore
 				}
 			}
 			// If nobody modified the package, we depend on the activity that added the package...
-			else if (const FActivityID* LastAddedActivity = PackageTracker.AddedPackages.Find(ModifiedPackage))
+			else if (LastAddedActivity)
 			{
 				AddDependency(*LastAddedActivity, EActivityDependencyReason::PackageCreation, EDependencyStrength::HardDependency);
 			}
 			// ... or "created" it by renaming it
-			else if (const FActivityID* LastSavedActivity = PackageTracker.RenamedPackages.Find(ModifiedPackage))
+			else if (LastRenamingActivity)
 			{
-				AddDependency(*LastSavedActivity, EActivityDependencyReason::PackageRename, EDependencyStrength::HardDependency);
+				AddDependency(*LastRenamingActivity, EActivityDependencyReason::PackageRename, EDependencyStrength::HardDependency);
 			}
 
 			/* We should consider indirect dependencies, too. Related ticket: UE-148392.
@@ -218,7 +236,7 @@ namespace UE::ConcertSyncCore
 	{
 		const FName& Package = EventData.PackageInfo.PackageName;
 		const FActivityID* LastPackageRemoval = PackageTracker.RemovedPackages.Find(Package);
-		const FActivityID* LastPackageRename = PackageTracker.RenamedPackages.Find(Package);
+		const FActivityID* LastPackageRename = PackageTracker.RenamedToPackages.Find(Package);
 		
 		AddDependencies(
 			NodeID,
@@ -232,47 +250,101 @@ namespace UE::ConcertSyncCore
 	{
 		const FName& Package = EventData.PackageInfo.PackageName;
 		const FActivityID* LastPackageAddition = PackageTracker.AddedPackages.Find(Package);
-		const FActivityID* LastPackageRename = PackageTracker.RenamedPackages.Find(Package);
+		const FActivityID* LastPackageRename = PackageTracker.RenamedToPackages.Find(Package);
 
 		const FTrackedPackageActivityArray Dependencies {
 			{ LastPackageAddition, EActivityDependencyReason::PackageCreation, EDependencyStrength::HardDependency, EPackageAddDependencyCondition::OnlyLatestActivity },
 			{ LastPackageRename, EActivityDependencyReason::PackageRename, EDependencyStrength::HardDependency, EPackageAddDependencyCondition::OnlyLatestActivity }
 		};
-		
 		AddDependencies(NodeID, Dependencies);
+
+		// Saving a package depends on the data that modified before it was saved
+		DiscoverPackageEditedDependencies(NodeID, Package);
 	}
 
 	void FActivityDependencyGraphBuildAlgorithm::DiscoverRenamedPackageDependencies(const FActivityNodeID NodeID, const FConcertSyncActivity& Activity, const FConcertSyncPackageEventMetaData& EventData)
 	{
 		const FName& NewPackageName = EventData.PackageInfo.NewPackageName;
 		const FName& OldPackageName = EventData.PackageInfo.PackageName;
-		
-		const FActivityID* NewNameAdded = PackageTracker.AddedPackages.Find(NewPackageName);
-		const FActivityID* NewPackageRenamed = PackageTracker.RenamedPackages.Find(NewPackageName);
+
+		// 1. Rename depends on whatever created the new package
+		const FActivityID* NewPackageAdded = PackageTracker.AddedPackages.Find(NewPackageName);
+		const FActivityID* NewPackageRenamed = PackageTracker.RenamedToPackages.Find(NewPackageName);
+		const FActivityID* LastActivityRenamingNewPackageName = PackageTracker.RenamedFromPackages.Find(NewPackageName);
 		const FActivityID* SavedNewPackage = PackageTracker.SavedPackages.Find(NewPackageName);
-		FTrackedPackageActivityArray NewPackageDependencies {
-			{ NewNameAdded, EActivityDependencyReason::PackageCreation, EDependencyStrength::HardDependency, EPackageAddDependencyCondition::OnlyLatestActivity },
-			{ NewPackageRenamed, EActivityDependencyReason::PackageRename, EDependencyStrength::HardDependency, EPackageAddDependencyCondition::OnlyLatestActivity },
-		};
+		const FActivityID* DeleteNewPackage = PackageTracker.RemovedPackages.Find(NewPackageName);
+
+		// An add package activity A is only a valid dependency if there was no delete between A and Activity.ActivityId. Similar for the other cases.
+		const bool bIsNewPackageActivityValid = NewPackageAdded
+			&& (!DeleteNewPackage || *DeleteNewPackage < *NewPackageAdded)
+			&& (!LastActivityRenamingNewPackageName || *LastActivityRenamingNewPackageName < *NewPackageAdded);
+		const bool bIsNewPackageRenamedValid = NewPackageRenamed
+			&& (!NewPackageAdded || *NewPackageAdded < *NewPackageRenamed)
+			&& (!LastActivityRenamingNewPackageName || *LastActivityRenamingNewPackageName < *NewPackageRenamed);
+		
+		FTrackedPackageActivityArray NewPackageDependencies;
 		// When renaming, Concert does not generate a EConcertPackageUpdateType::Added type.
 		// Instead the renamed to package is added via an EConcertPackageUpdateType::Saved activity which occurs just before
 		// the corresponding EConcertPackageUpdateType::Rename activity.
-		if (!NewNameAdded && !NewPackageRenamed)
+		if (!bIsNewPackageActivityValid && !bIsNewPackageRenamedValid)
+		{
+			// Maybe this could be a hard dependency...
+			NewPackageDependencies = {{ SavedNewPackage, EActivityDependencyReason::PackageCreation, EDependencyStrength::HardDependency, EPackageAddDependencyCondition::Always }};
+			
+			// Saving a package depends on the data that modified before it was saved
+			// SavedNewPackage was already processed but since the algorithm only looks backwards we didn't know the save operation was needed by a rename.
+			// Now we know and must retrospectively update that nodes dependencies.
+			if (SavedNewPackage)
+			{
+				const TOptional<FActivityNodeID> SaveNodeId = Graph.FindNodeByActivity(*SavedNewPackage);
+				checkf(SaveNodeId.IsSet(), TEXT("SaveNodeId (%lld) preceeds the rename (%lld) so we should already have created a node for it in the graph!"), *SavedNewPackage, Activity.ActivityId);
+				DiscoverPackageEditedDependencies(*SaveNodeId, OldPackageName);
+			}
+		}
+		else
 		{
-			NewPackageDependencies.Add({ SavedNewPackage, EActivityDependencyReason::PackageCreation, EDependencyStrength::PossibleDependency, EPackageAddDependencyCondition::Always });
+			if (bIsNewPackageActivityValid)
+			{
+				NewPackageDependencies.Add({ NewPackageAdded, EActivityDependencyReason::PackageCreation, EDependencyStrength::HardDependency, EPackageAddDependencyCondition::OnlyLatestActivity });
+			}
+			if (bIsNewPackageRenamedValid)
+			{
+				NewPackageDependencies.Add({ NewPackageRenamed, EActivityDependencyReason::PackageRename, EDependencyStrength::HardDependency, EPackageAddDependencyCondition::OnlyLatestActivity });
+			}
 		}
-		
+
+		// 2. Rename depends on whatever created the old package
 		const FActivityID* OldPackageAdded = PackageTracker.AddedPackages.Find(OldPackageName);
-		const FActivityID* OldPackageRenamed = PackageTracker.RenamedPackages.Find(OldPackageName);
+		const FActivityID* OldPackageRenamed = PackageTracker.RenamedToPackages.Find(OldPackageName);
+		const FActivityID* LastActivityRenamingOldPackageName = PackageTracker.RenamedFromPackages.Find(OldPackageName);
 		const FActivityID* SavedOldPackage = PackageTracker.SavedPackages.Find(OldPackageName);
-		FTrackedPackageActivityArray OldPackageDependencies {
-			{ OldPackageAdded, EActivityDependencyReason::PackageCreation, EDependencyStrength::HardDependency, EPackageAddDependencyCondition::OnlyLatestActivity },
-			{ OldPackageRenamed, EActivityDependencyReason::PackageRename, EDependencyStrength::HardDependency, EPackageAddDependencyCondition::OnlyLatestActivity },
-		};
+		const FActivityID* DeleteOldPackage = PackageTracker.RemovedPackages.Find(OldPackageName);
+
+		// See similar cases from above
+		const bool bIsOldPackageActivityValid = OldPackageAdded
+			&& (!DeleteOldPackage || *DeleteOldPackage < *OldPackageAdded)
+			&& (!LastActivityRenamingOldPackageName || *LastActivityRenamingOldPackageName < *OldPackageAdded);
+		const bool bIsOldPackageRenamedValid = OldPackageRenamed
+			&& (!OldPackageAdded || *OldPackageAdded < *OldPackageRenamed)
+			&& (!LastActivityRenamingOldPackageName || *LastActivityRenamingOldPackageName < *OldPackageRenamed);
+		
+		FTrackedPackageActivityArray OldPackageDependencies;
 		// See above
-		if (!OldPackageAdded && !OldPackageRenamed)
+		if (!bIsOldPackageActivityValid && !bIsOldPackageRenamedValid)
 		{
-			OldPackageDependencies.Add({ SavedOldPackage, EActivityDependencyReason::PackageCreation, EDependencyStrength::PossibleDependency, EPackageAddDependencyCondition::Always });
+			// Maybe this could be a hard dependency...
+			OldPackageDependencies = {{ SavedOldPackage, EActivityDependencyReason::PackageCreation, EDependencyStrength::HardDependency, EPackageAddDependencyCondition::Always }};
+		}
+		else
+		{
+			if (bIsOldPackageActivityValid)
+			{
+				OldPackageDependencies.Add({ OldPackageAdded, EActivityDependencyReason::PackageCreation, EDependencyStrength::HardDependency, EPackageAddDependencyCondition::OnlyLatestActivity });
+			}
+			if (bIsOldPackageRenamedValid)
+			{
+				OldPackageDependencies.Add({ OldPackageRenamed, EActivityDependencyReason::PackageRename, EDependencyStrength::HardDependency, EPackageAddDependencyCondition::OnlyLatestActivity });
+			}
 		}
 		
 		AddDependencies(NodeID, NewPackageDependencies);
@@ -283,7 +355,7 @@ namespace UE::ConcertSyncCore
 	{
 		const FName& Package = EventData.PackageInfo.PackageName;
 		const FActivityID* LastPackageAddition = PackageTracker.AddedPackages.Find(Package);
-		const FActivityID* LastPackageRename = PackageTracker.RenamedPackages.Find(Package);
+		const FActivityID* LastPackageRename = PackageTracker.RenamedToPackages.Find(Package);
 		
 		AddDependencies(
 			NodeID,
@@ -293,6 +365,19 @@ namespace UE::ConcertSyncCore
 			});
 	}
 	
+	void FActivityDependencyGraphBuildAlgorithm::DiscoverPackageEditedDependencies(const FActivityNodeID DependingNodeID, FName ModifiedPackage)
+	{
+		TArray<FActivityID> TransactionsAffectingPackage;
+		PackageTracker.TransactionsAffectingPackageSinceLastPackageActivity.RemoveAndCopyValue(ModifiedPackage, TransactionsAffectingPackage);
+		for (const FActivityID& TransactionActivityID : TransactionsAffectingPackage)
+		{
+			AddDependencies(
+				DependingNodeID, 
+				{{ &TransactionActivityID, EActivityDependencyReason::PackageEdited, EDependencyStrength::HardDependency, EPackageAddDependencyCondition::Always }}
+			);
+		}
+	}
+	
 	void FActivityDependencyGraphBuildAlgorithm::AddDependencies(const FActivityNodeID NodeID, const FTrackedPackageActivityArray& Dependencies)
 	{
 		auto AddDependency = [this, NodeID](const FPackageActivityItem& Dependency)
@@ -343,6 +428,10 @@ namespace UE::ConcertSyncCore
 		for (FName ModifiedPackage : EventData.Transaction.ModifiedPackages)
 		{
 			PackageTracker.ModifiedPackages.Add(ModifiedPackage, Activity.ActivityId);
+			
+			PackageTracker.TransactionsAffectingPackageSinceLastPackageActivity
+				.FindOrAdd(ModifiedPackage)
+				.Add(Activity.ActivityId);
 		}
 		
 		for (const FConcertExportedObject& ExportedObject : EventData.Transaction.ExportedObjects)
@@ -370,14 +459,17 @@ namespace UE::ConcertSyncCore
 			
 		case EConcertPackageUpdateType::Saved:
 			PackageTracker.SavedPackages.Add(PackageName, ActivityID);
+			PackageTracker.TransactionsAffectingPackageSinceLastPackageActivity.Remove(PackageName);
 			break;
 			
 		case EConcertPackageUpdateType::Renamed:
-			PackageTracker.RenamedPackages.Add(EventData.PackageInfo.NewPackageName, ActivityID);
+			PackageTracker.RenamedToPackages.Add(EventData.PackageInfo.NewPackageName, ActivityID);
+			PackageTracker.RenamedFromPackages.Add(EventData.PackageInfo.PackageName, ActivityID);
 			break;
 			
 		case EConcertPackageUpdateType::Deleted:
 			PackageTracker.RemovedPackages.Add(PackageName, ActivityID);
+			PackageTracker.TransactionsAffectingPackageSinceLastPackageActivity.Remove(PackageName);
 			break;
 			
 		case EConcertPackageUpdateType::Dummy:
@@ -388,7 +480,7 @@ namespace UE::ConcertSyncCore
 			checkNoEntry();
 		}
 	}
-
+	
 	FSoftObjectPath FActivityDependencyGraphBuildAlgorithm::MakePathFromExportedObject(const FConcertExportedObject& ExportedObject) const
 	{
 		return FSoftObjectPath(FString::Printf(TEXT("%s.%s"), *ExportedObject.ObjectId.ObjectOuterPathName.ToString(), *ExportedObject.ObjectId.ObjectName.ToString()));
diff --git a/Engine/Plugins/Developer/Concert/ConcertSync/ConcertSyncCore/Source/ConcertSyncCore/Private/HistoryEdition/HistoryTopDownAnalysis.cpp b/Engine/Plugins/Developer/Concert/ConcertSync/ConcertSyncCore/Source/ConcertSyncCore/Private/HistoryEdition/HistoryTopDownAnalysis.cpp
--- a/Engine/Plugins/Developer/Concert/ConcertSync/ConcertSyncCore/Source/ConcertSyncCore/Private/HistoryEdition/HistoryTopDownAnalysis.cpp
+++ b/Engine/Plugins/Developer/Concert/ConcertSync/ConcertSyncCore/Source/ConcertSyncCore/Private/HistoryEdition/HistoryTopDownAnalysis.cpp
@@ -137,10 +137,10 @@ namespace UE::ConcertSyncCore
 		const FActivityDependencyEdge* PackageCreationDependency = Algo::FindByPredicate(PossibleRenameActivityNode.GetDependencies(),
 			[](const FActivityDependencyEdge& Dependency)
 			{
-				const bool bIsAssociatedCreationActivity = Dependency.GetDependencyReason() == EActivityDependencyReason::PackageCreation && Dependency.GetDependencyStrength() == EDependencyStrength::PossibleDependency;
+				const bool bIsAssociatedCreationActivity = Dependency.GetDependencyReason() == EActivityDependencyReason::PackageCreation && Dependency.GetDependencyStrength() == EDependencyStrength::HardDependency;
 				return bIsAssociatedCreationActivity;
 			});
-		if (ensureMsgf(PackageCreationDependency, TEXT("Rename activities should have a possible dependency to the activity that creatd the renamed-to package"))
+		if (ensureMsgf(PackageCreationDependency, TEXT("Rename activities should have a hard dependency to the activity that creatd the renamed-to package"))
 			&& !DoubleEnqueuingProtection.Contains(PackageCreationDependency->GetDependedOnNodeID()))
 		{
 			DoubleEnqueuingProtection.Add(PackageCreationDependency->GetDependedOnNodeID());
diff --git a/Engine/Plugins/Developer/Concert/ConcertSync/ConcertSyncCore/Source/ConcertSyncCore/Public/HistoryEdition/ActivityDependencyEdge.h b/Engine/Plugins/Developer/Concert/ConcertSync/ConcertSyncCore/Source/ConcertSyncCore/Public/HistoryEdition/ActivityDependencyEdge.h
--- a/Engine/Plugins/Developer/Concert/ConcertSync/ConcertSyncCore/Source/ConcertSyncCore/Public/HistoryEdition/ActivityDependencyEdge.h
+++ b/Engine/Plugins/Developer/Concert/ConcertSync/ConcertSyncCore/Source/ConcertSyncCore/Public/HistoryEdition/ActivityDependencyEdge.h
@@ -40,6 +40,14 @@ namespace UE::ConcertSyncCore
 		 */
 		PackageRename,
 
+		/**
+		 * The depended-on activity edits the package that is required by this activity.
+		 *
+		 * Example: 1: Move actor 2: Save level
+		 * 2 depends on 1 (because the save operation would save different data if 1 was not performed).
+		 */
+		PackageEdited,
+
 		/**
 		 * The depended-on activity changes the state of a package that is depended on by a package this activity modifies.
 		 * 
@@ -62,7 +70,10 @@ namespace UE::ConcertSyncCore
 		 * Example: 1: Remove actor > 2. Create actor with same name
 		 * 2 cannot take place without 1.
 		 */
-		SubobjectRemoval
+		SubobjectRemoval,
+
+		// ADD BEFORE THIS ENTRY
+		Count
 	};
 	
 	FString LexToString(EActivityDependencyReason Reason);
@@ -76,7 +87,7 @@ namespace UE::ConcertSyncCore
 		HardDependency,
 		
 		/**
-		 * An activity A has a possible dependency to activity B when A may exist with B but we cannot rule out
+		 * An activity A has a possible dependency to activity B when A could possibly exist without B but we cannot rule out
 		 * that B does not affect A.
 		 * Example: You modify an actor twice each time triggering the construction script
 		 */
diff --git a/Engine/Plugins/Developer/Concert/ConcertSync/ConcertSyncTest/Source/ConcertSyncTest/Private/HistoryEdition/DependencyGraphBuilderTests.cpp b/Engine/Plugins/Developer/Concert/ConcertSync/ConcertSyncTest/Source/ConcertSyncTest/Private/HistoryEdition/DependencyGraphBuilderTests.cpp
--- a/Engine/Plugins/Developer/Concert/ConcertSync/ConcertSyncTest/Source/ConcertSyncTest/Private/HistoryEdition/DependencyGraphBuilderTests.cpp
+++ b/Engine/Plugins/Developer/Concert/ConcertSync/ConcertSyncTest/Source/ConcertSyncTest/Private/HistoryEdition/DependencyGraphBuilderTests.cpp
@@ -5,6 +5,7 @@
 #include "ConcertSyncSessionDatabase.h"
 #include "HistoryEdition/DebugDependencyGraph.h"
 #include "HistoryEdition/DependencyGraphBuilder.h"
+#include "HistoryTestUtil.h"
 #include "RenameEditAndDeleteMapsFlow.h"
 #include "ScopedSessionDatabase.h"
 
@@ -54,79 +55,84 @@ namespace UE::ConcertSyncTests::RenameEditAndDeleteMapsFlowTest
 	{
 		using namespace ConcertSyncCore;
 		
-		const TTestActivityArray<FActivityNodeID> ActivityNodes
+		const TTestActivityArray<FActivityNodeID> Nodes
 			= ValidateEachActivityHasNode(Test, ActivityMappings, Graph);
 
+		// TODO: Update this to respect PackageEdited
+		
 		// 1 Create map Foo
 		{
-			Test.TestFalse(TEXT("1 Creating new package 'Foo' has no dependencies."), Graph.GetNodeById(ActivityNodes[_1_NewPackageFoo]).HasAnyDependency());
-			Test.TestTrue(TEXT("1 Creating new package 'Foo' has correct node flags"), Graph.GetNodeById(ActivityNodes[_1_NewPackageFoo]).GetNodeFlags() == EActivityNodeFlags::None);
+			Test.TestFalse(TEXT("_1_NewPackageFoo has no dependencies"), Graph.GetNodeById(Nodes[_1_NewPackageFoo]).HasAnyDependency());
+			Test.TestTrue(TEXT("_1_NewPackageFoo has correct node flags"), Graph.GetNodeById(Nodes[_1_NewPackageFoo]).GetNodeFlags() == EActivityNodeFlags::None);
 			
-			Test.TestTrue(TEXT("1 Saving 'Foo' has dependency to creating package 'Foo'."), Graph.GetNodeById(ActivityNodes[_1_SavePackageFoo]).DependsOnActivity(ActivityMappings[_1_NewPackageFoo], Graph, EActivityDependencyReason::PackageCreation, EDependencyStrength::HardDependency));
-			Test.TestTrue(TEXT("1 Saving 'Foo' has correct node flags"), Graph.GetNodeById(ActivityNodes[_1_SavePackageFoo]).GetNodeFlags() == EActivityNodeFlags::None);
-			Test.TestEqual(TEXT("1 Saving 'Foo' has exactly 1 dependency"), Graph.GetNodeById(ActivityNodes[_1_SavePackageFoo]).GetDependencies().Num(), 1);
+			Test.TestTrue(TEXT("_2_SavePackageFoo depends on _1_NewPackageFoo"), Graph.GetNodeById(Nodes[_2_SavePackageFoo]).DependsOnActivity(ActivityMappings[_1_NewPackageFoo], Graph, EActivityDependencyReason::PackageCreation, EDependencyStrength::HardDependency));
+			Test.TestTrue(TEXT("_2_SavePackageFoo has correct node flags"), Graph.GetNodeById(Nodes[_2_SavePackageFoo]).GetNodeFlags() == EActivityNodeFlags::None);
+			Test.TestEqual(TEXT("_2_SavePackageFoo only depends on _1_NewPackageFoo"), Graph.GetNodeById(Nodes[_2_SavePackageFoo]).GetDependencies().Num(), 1);
 		}
 
 		// 2 Add actor A
 		{
-			Test.TestTrue(TEXT("2 Adding actor to 'Foo' depends on creating package 'Foo'."), Graph.GetNodeById(ActivityNodes[_2_AddActor]).DependsOnActivity(ActivityMappings[_1_NewPackageFoo], Graph, EActivityDependencyReason::PackageCreation, EDependencyStrength::HardDependency));
-			Test.TestTrue(TEXT("2 Adding actor to 'Foo' has correct node flags"), Graph.GetNodeById(ActivityNodes[_2_AddActor]).GetNodeFlags() == EActivityNodeFlags::None);
-			Test.TestEqual(TEXT("2 Saving 'Foo' has exactly 1 dependency"), Graph.GetNodeById(ActivityNodes[_2_AddActor]).GetDependencies().Num(), 1);
+			Test.TestTrue(TEXT("_3_AddActor depends on _1_NewPackageFoo"), Graph.GetNodeById(Nodes[_3_AddActor]).DependsOnActivity(ActivityMappings[_1_NewPackageFoo], Graph, EActivityDependencyReason::PackageCreation, EDependencyStrength::HardDependency));
+			Test.TestTrue(TEXT("_3_AddActor has correct node flags"), Graph.GetNodeById(Nodes[_3_AddActor]).GetNodeFlags() == EActivityNodeFlags::None);
+			Test.TestEqual(TEXT("_3_AddActor only depends on _1_NewPackageFoo"), Graph.GetNodeById(Nodes[_3_AddActor]).GetDependencies().Num(), 1);
 		}
 
 		// 3 Rename actor A
 		{
 			// It must be a EDependencyStrength::HardDependency because you cannot edit the actor without having created it 
-			Test.TestTrue(TEXT("3 Renaming actor depends on having created the actor."), Graph.GetNodeById(ActivityNodes[_3_RenameActor]).DependsOnActivity(ActivityMappings[_2_AddActor], Graph, EActivityDependencyReason::SubobjectCreation, EDependencyStrength::HardDependency));
-			Test.TestTrue(TEXT("3 Renaming actor to 'Foo' has correct node flags"), Graph.GetNodeById(ActivityNodes[_3_RenameActor]).GetNodeFlags() == EActivityNodeFlags::None);
-			Test.TestEqual(TEXT("3 Renaming actor has exactly 1 dependency"), Graph.GetNodeById(ActivityNodes[_3_RenameActor]).GetDependencies().Num(), 1);
+			Test.TestTrue(TEXT("_4_RenameActor depends on _3_AddActor"), Graph.GetNodeById(Nodes[_4_RenameActor]).DependsOnActivity(ActivityMappings[_3_AddActor], Graph, EActivityDependencyReason::SubobjectCreation, EDependencyStrength::HardDependency));
+			Test.TestTrue(TEXT("_4_RenameActor has correct node flags"), Graph.GetNodeById(Nodes[_4_RenameActor]).GetNodeFlags() == EActivityNodeFlags::None);
+			Test.TestEqual(TEXT("_4_RenameActor only depends on _3_AddActor"), Graph.GetNodeById(Nodes[_4_RenameActor]).GetDependencies().Num(), 1);
 		}
 
 		// 4 Edit actor A
 		{
+			// This activity must have a EDependencyStrength::HardDependency to _2_AddActor because the edit cannot happen without having created the actor
+			Test.TestTrue(TEXT("_5_EditActor depends on _3_AddActor"), Graph.GetNodeById(Nodes[_5_EditActor]).DependsOnActivity(ActivityMappings[_3_AddActor], Graph, EActivityDependencyReason::SubobjectCreation, EDependencyStrength::HardDependency));
 			// The previous edit might have affected us (e.g. this activity may have executed the construction script)
 			// Note: This should not have a hard dependency on having renamed the actor because a rename is just a property change of ActorLabel.
-			Test.TestTrue(TEXT("4 Editing actor may depend on having edited the actor previously."), Graph.GetNodeById(ActivityNodes[_4_EditActor]).DependsOnActivity(ActivityMappings[_3_RenameActor], Graph, EActivityDependencyReason::EditAfterPreviousPackageEdit, EDependencyStrength::PossibleDependency));
-			// This activity must have a EDependencyStrength::HardDependency to _2_AddActor because the edit cannot happen without having created the actor
-			Test.TestTrue(TEXT("4 Editing actor depends on having created the actor."), Graph.GetNodeById(ActivityNodes[_4_EditActor]).DependsOnActivity(ActivityMappings[_2_AddActor], Graph, EActivityDependencyReason::SubobjectCreation, EDependencyStrength::HardDependency));
-			Test.TestTrue(TEXT("4 Editing actor has correct node flags"), Graph.GetNodeById(ActivityNodes[_4_EditActor]).GetNodeFlags() == EActivityNodeFlags::None);
-			Test.TestEqual(TEXT("4 Editing actor has exactly 1 dependency"), Graph.GetNodeById(ActivityNodes[_4_EditActor]).GetDependencies().Num(), 2);
+			Test.TestTrue(TEXT("_5_EditActor depends on _4_RenameActor"), Graph.GetNodeById(Nodes[_5_EditActor]).DependsOnActivity(ActivityMappings[_4_RenameActor], Graph, EActivityDependencyReason::EditAfterPreviousPackageEdit, EDependencyStrength::PossibleDependency));
+			Test.TestTrue(TEXT("_5_EditActor has correct node flags"), Graph.GetNodeById(Nodes[_5_EditActor]).GetNodeFlags() == EActivityNodeFlags::None);
+			Test.TestEqual(TEXT("_5_EditActor only depends on _3_AddActor and _4_RenameActor"), Graph.GetNodeById(Nodes[_5_EditActor]).GetDependencies().Num(), 2);
 		}
 
 		// 5 Rename map to Bar
 		{
-			Test.TestFalse(TEXT("5 Saving new package 'Bar' has no dependencies."), Graph.GetNodeById(ActivityNodes[_5_SavePackageBar]).HasAnyDependency());
-			Test.TestTrue(TEXT("5 Saving new package 'Bar' has correct node flags"), Graph.GetNodeById(ActivityNodes[_5_SavePackageBar]).GetNodeFlags() == EActivityNodeFlags::None);
+			Test.TestTrue(TEXT("_6_SavePackageBar has correct node flags"), Graph.GetNodeById(Nodes[_6_SavePackageBar]).GetNodeFlags() == EActivityNodeFlags::None);
+			Test.TestTrue(TEXT("_6_SavePackageBar depends on _3_AddActor"), Graph.GetNodeById(Nodes[_6_SavePackageBar]).DependsOnNode(Nodes[_3_AddActor], EActivityDependencyReason::PackageEdited, EDependencyStrength::HardDependency));
+			Test.TestTrue(TEXT("_6_SavePackageBar depends on _4_RenameActor"), Graph.GetNodeById(Nodes[_6_SavePackageBar]).DependsOnNode(Nodes[_4_RenameActor], EActivityDependencyReason::PackageEdited, EDependencyStrength::HardDependency));
+			Test.TestTrue(TEXT("_6_SavePackageBar depends on _5_EditActor"), Graph.GetNodeById(Nodes[_6_SavePackageBar]).DependsOnNode(Nodes[_5_EditActor], EActivityDependencyReason::PackageEdited, EDependencyStrength::HardDependency));
+			Test.TestEqual(TEXT("_6_SavePackageBar depends on _3_AddActor, _4_RenameActor, and _5_EditActor"), Graph.GetNodeById(Nodes[_6_SavePackageBar]).GetDependencies().Num(), 3);
 			
-			Test.TestTrue(TEXT("5 Renaming 'Foo' to 'Bar' has dependency to creating package 'Foo'."), Graph.GetNodeById(ActivityNodes[_5_RenameFooToBar]).DependsOnActivity(ActivityMappings[_1_NewPackageFoo], Graph, EActivityDependencyReason::PackageCreation, EDependencyStrength::HardDependency));
-			Test.TestTrue(TEXT("5 Renaming 'Foo' to 'Bar' has dependency to creating package 'Bar'."), Graph.GetNodeById(ActivityNodes[_5_RenameFooToBar]).DependsOnActivity(ActivityMappings[_5_SavePackageBar], Graph, EActivityDependencyReason::PackageCreation, EDependencyStrength::PossibleDependency));
-			Test.TestTrue(TEXT("5 Renaming 'Foo' to 'Bar' has correct node flags"), Graph.GetNodeById(ActivityNodes[_5_RenameFooToBar]).GetNodeFlags() == EActivityNodeFlags::RenameActivity);
-			Test.TestEqual(TEXT("5 Saving 'Foo' has exactly 2 dependencies"), Graph.GetNodeById(ActivityNodes[_5_RenameFooToBar]).GetDependencies().Num(), 2);
+			Test.TestTrue(TEXT("_7_RenameFooToBar depends on _1_NewPackageFoo"), Graph.GetNodeById(Nodes[_7_RenameFooToBar]).DependsOnActivity(ActivityMappings[_1_NewPackageFoo], Graph, EActivityDependencyReason::PackageCreation, EDependencyStrength::HardDependency));
+			Test.TestTrue(TEXT("_7_RenameFooToBar depends on _6_SavePackageBar"), Graph.GetNodeById(Nodes[_7_RenameFooToBar]).DependsOnActivity(ActivityMappings[_6_SavePackageBar], Graph, EActivityDependencyReason::PackageCreation, EDependencyStrength::HardDependency));
+			Test.TestTrue(TEXT("_7_RenameFooToBar has correct node flags"), Graph.GetNodeById(Nodes[_7_RenameFooToBar]).GetNodeFlags() == EActivityNodeFlags::RenameActivity);
+			Test.TestEqual(TEXT("_7_RenameFooToBar depends only on _1_NewPackageFoo and _6_SavePackageBar"), Graph.GetNodeById(Nodes[_7_RenameFooToBar]).GetDependencies().Num(), 2);
 		}
 
 		// 6 Edit actor A
 		{
-			Test.TestTrue(TEXT("6 Editing actor in 'Bar' depends on having renamed 'Foo' to 'Bar'."), Graph.GetNodeById(ActivityNodes[_6_EditActor]).DependsOnActivity(ActivityMappings[_5_RenameFooToBar], Graph, EActivityDependencyReason::PackageRename, EDependencyStrength::HardDependency));
-			Test.TestTrue(TEXT("6 Editing actor has correct node flags"), Graph.GetNodeById(ActivityNodes[_6_EditActor]).GetNodeFlags() == EActivityNodeFlags::None);
-			Test.TestEqual(TEXT("6 Editing actor has exactly 1 dependency"), Graph.GetNodeById(ActivityNodes[_6_EditActor]).GetDependencies().Num(), 1);
+			Test.TestTrue(TEXT("_8_EditActor depends on _7_RenameFooToBar"), Graph.GetNodeById(Nodes[_8_EditActor]).DependsOnActivity(ActivityMappings[_7_RenameFooToBar], Graph, EActivityDependencyReason::PackageRename, EDependencyStrength::HardDependency));
+			Test.TestTrue(TEXT("_8_EditActor has correct node flags"), Graph.GetNodeById(Nodes[_8_EditActor]).GetNodeFlags() == EActivityNodeFlags::None);
+			Test.TestEqual(TEXT("_8_EditActor only depends on _7_RenameFooToBar"), Graph.GetNodeById(Nodes[_8_EditActor]).GetDependencies().Num(), 1);
 		}
 
 		// 7 Delete map Bar
 		{
-			Test.TestTrue(TEXT("7 Deleting 'Bar' depends on having renamed 'Foo' to 'Bar' previously."), Graph.GetNodeById(ActivityNodes[_7_DeleteBar]).DependsOnActivity(ActivityMappings[_5_RenameFooToBar], Graph, EActivityDependencyReason::PackageRename, EDependencyStrength::HardDependency));
-			Test.TestTrue(TEXT("7 Deleting 'Bar' has correct node flags"), Graph.GetNodeById(ActivityNodes[_7_DeleteBar]).GetNodeFlags() == EActivityNodeFlags::None);
-			Test.TestEqual(TEXT("7 Deleting 'Bar' after rename has exactly 1 dependency"), Graph.GetNodeById(ActivityNodes[_7_DeleteBar]).GetDependencies().Num(), 1);
+			Test.TestTrue(TEXT("_9_DeleteBar depends on _7_RenameFooToBar"), Graph.GetNodeById(Nodes[_9_DeleteBar]).DependsOnActivity(ActivityMappings[_7_RenameFooToBar], Graph, EActivityDependencyReason::PackageRename, EDependencyStrength::HardDependency));
+			Test.TestTrue(TEXT("_9_DeleteBar has correct node flags"), Graph.GetNodeById(Nodes[_9_DeleteBar]).GetNodeFlags() == EActivityNodeFlags::None);
+			Test.TestEqual(TEXT("_9_DeleteBar depends only on _7_RenameFooToBar"), Graph.GetNodeById(Nodes[_9_DeleteBar]).GetDependencies().Num(), 1);
 		}
 
 		// 8 Create map Bar
 		{
-			Test.TestTrue(TEXT("8 Re-creating 'Bar' depends on having deleted 'Bar' previously."), Graph.GetNodeById(ActivityNodes[_8_NewPackageFoo]).DependsOnActivity(ActivityMappings[_7_DeleteBar], Graph, EActivityDependencyReason::PackageRemoval, EDependencyStrength::HardDependency));
-			Test.TestTrue(TEXT("8 Re-creating actor has correct node flags"), Graph.GetNodeById(ActivityNodes[_8_NewPackageFoo]).GetNodeFlags() == EActivityNodeFlags::None);
-			Test.TestEqual(TEXT("8 Re-creating 'Bar' 1 dependency"), Graph.GetNodeById(ActivityNodes[_8_NewPackageFoo]).GetDependencies().Num(), 1);
+			Test.TestTrue(TEXT("_10_NewPackageFoo depends on _9_DeleteBar"), Graph.GetNodeById(Nodes[_10_NewPackageFoo]).DependsOnActivity(ActivityMappings[_9_DeleteBar], Graph, EActivityDependencyReason::PackageRemoval, EDependencyStrength::HardDependency));
+			Test.TestTrue(TEXT("_10_NewPackageFoo has correct node flags"), Graph.GetNodeById(Nodes[_10_NewPackageFoo]).GetNodeFlags() == EActivityNodeFlags::None);
+			Test.TestEqual(TEXT("_10_NewPackageFoo depends only on _9_DeleteBar"), Graph.GetNodeById(Nodes[_10_NewPackageFoo]).GetDependencies().Num(), 1);
 			
-			Test.TestTrue(TEXT("8 Saving 'Bar' depends on re-created 'Bar'."), Graph.GetNodeById(ActivityNodes[_8_SavePackageFoo]).DependsOnActivity(ActivityMappings[_8_NewPackageFoo], Graph, EActivityDependencyReason::PackageCreation, EDependencyStrength::HardDependency));
-			Test.TestTrue(TEXT("8 Saving 'Bar' has correct node flags"), Graph.GetNodeById(ActivityNodes[_8_SavePackageFoo]).GetNodeFlags() == EActivityNodeFlags::None);
-			Test.TestEqual(TEXT("8 Saving 'Bar' after re-creation has exactly 1 dependency"), Graph.GetNodeById(ActivityNodes[_8_SavePackageFoo]).GetDependencies().Num(), 1);
+			Test.TestTrue(TEXT("_11_SavePackageFoo depends on _10_NewPackageFoo"), Graph.GetNodeById(Nodes[_11_SavePackageFoo]).DependsOnActivity(ActivityMappings[_10_NewPackageFoo], Graph, EActivityDependencyReason::PackageCreation, EDependencyStrength::HardDependency));
+			Test.TestTrue(TEXT("_11_SavePackageFoo has correct node flags"), Graph.GetNodeById(Nodes[_11_SavePackageFoo]).GetNodeFlags() == EActivityNodeFlags::None);
+			Test.TestEqual(TEXT("_11_SavePackageFoo depends only on _10_NewPackageFoo"), Graph.GetNodeById(Nodes[_11_SavePackageFoo]).GetDependencies().Num(), 1);
 		}
 		
 		return true;
@@ -306,4 +312,342 @@ namespace UE::ConcertSyncTests::IndirectPackageDependencyTest
 		// TODO:
 		return true;
 	}
+}
+
+namespace UE::ConcertSyncTests::PackageEditedDependencyTest
+{
+	enum ETestActivity
+	{
+		// Normal flow: saving depends on transaction
+		_1_CreateActor,
+		_2_SaveFoo,
+		// Saving after without transactions since last save has no dependencies
+		_3_SaveFoo,
+		_4_CreateActor,
+
+		// Renaming a package inherits the dependencies to the transactions that occured before the rename
+		_5_SaveBar,
+		_6_RenameFooToBar,
+
+		// Creating an equally named package after renaming it causes no dependencies
+		_7_NewMapFoo,
+		_8_SaveMapFoo,
+		_9_CreateActor,
+		_10_SaveFoo,
+		_11_DeleteFoo,
+
+		// Creating an equally named package after deleting it causes no dependencies
+		_12_NewMapFoo,
+		_13_SaveMapFoo,
+		_14_CreateActor,
+		_15_SaveFoo,
+
+		ActivityCount
+	};
+	
+	TArray<FActivityID> FillDatabase(FScopedSessionDatabase& SessionDatabase);
+	FConcertExportedObject CreateActorMetaData(FName OuterLevelPath, FName ActorName);
+	
+	/**
+	 * This tests that EActivityDependencyReason::PackageEdited dependencies are discovered correctly.
+	 *
+	 * Sequence of user actions:
+	 *  1 Create actor
+	 *  2 Save
+	 *  3 Save
+	 *  4 Create actor
+	 *  5 Rename to Bar
+	 *  6 Create new map Foo
+	 *  7 Create actor
+	 *  8 Save Foo
+	 *  9 Delete Foo
+	 *  10 New map Foo
+	 *  11 Save Foo
+	 *  12 Create actor
+	 *  13 Save Foo
+	 *
+	 *  Only these dependencies should have PackageEdited as reason:
+	 *  2 -> 1
+	 *  5 -> 4 (the save before the rename has the dependency)
+	 *  8 -> 7
+	 *  13 -> 12
+	 */
+	IMPLEMENT_SIMPLE_AUTOMATION_TEST(FPackageEditedDependencyTest, "Concert.History.BuildGraph.PackageEditedDependencyTest", EAutomationTestFlags::ApplicationContextMask | EAutomationTestFlags::EngineFilter)
+	bool FPackageEditedDependencyTest::RunTest(const FString& Parameters)
+	{
+		using namespace ConcertSyncCore;
+		
+		FScopedSessionDatabase SessionDatabase(*this);
+		TArray<FActivityID> Activities = FillDatabase(SessionDatabase);
+		
+		const FActivityDependencyGraph DependencyGraph = BuildDependencyGraphFrom(SessionDatabase);
+		UE_LOG(LogConcert, Log, TEXT("%s tested graph in Graphviz format:\n\n%s"), *GetTestFullName(), *ConcertSyncCore::Graphviz::ExportToGraphviz(DependencyGraph, SessionDatabase));
+		TArray<FActivityNodeID> Nodes = ConcertSyncTests::ValidateEachActivityHasNode(*this, Activities, DependencyGraph, ActivityCount, [](const uint32 ActivityType){ return FString::FromInt(ActivityType); });
+
+		// _1_CreateActor
+		{
+			TestTrue(TEXT("_1_CreateActor: Affects first _2_SaveFoo"), DependencyGraph.GetNodeById(Nodes[_1_CreateActor]).AffectsActivity(Activities[_2_SaveFoo], DependencyGraph));
+			TestTrue(TEXT("_1_CreateActor: Affects first _2_SaveFoo"), DependencyGraph.GetNodeById(Nodes[_1_CreateActor]).AffectsActivity(Activities[_4_CreateActor], DependencyGraph));
+			TestEqual(TEXT("_1_CreateActor: Affects only _2_SaveFoo"), DependencyGraph.GetNodeById(Nodes[_1_CreateActor]).GetAffectedChildren().Num(), 2);
+		}
+		// _2_SaveFoo
+		{
+			TestTrue(TEXT("_2_SaveFoo: Depends on _1_CreateActor"), DependencyGraph.GetNodeById(Nodes[_2_SaveFoo]).DependsOnNode(Nodes[_1_CreateActor], EActivityDependencyReason::PackageEdited, EDependencyStrength::HardDependency));
+			TestEqual(TEXT("_2_SaveFoo: Depends only on _1_CreateActor"), DependencyGraph.GetNodeById(Nodes[_2_SaveFoo]).GetDependencies().Num(), 1);
+		}
+		// _3_SaveFoo
+		{
+			TestEqual(TEXT("_3_SaveFoo: Depends on nothing"), DependencyGraph.GetNodeById(Nodes[_3_SaveFoo]).GetDependencies().Num(), 0);
+		}
+		// _4_CreateActor
+		{
+			TestTrue(TEXT("_4_CreateActor: Depends on _1_CreateActor"), DependencyGraph.GetNodeById(Nodes[_4_CreateActor]).DependsOnNode(Nodes[_1_CreateActor], EActivityDependencyReason::EditAfterPreviousPackageEdit, EDependencyStrength::PossibleDependency));
+			TestEqual(TEXT("_4_CreateActor: Depends only on _1_CreateActor"), DependencyGraph.GetNodeById(Nodes[_4_CreateActor]).GetDependencies().Num(), 1);
+		}
+		// _5_SaveFoo2
+		{
+			TestTrue(TEXT("_5_SaveFoo2: Depends on _4_CreateActor"), DependencyGraph.GetNodeById(Nodes[_5_SaveBar]).DependsOnNode(Nodes[_4_CreateActor], EActivityDependencyReason::PackageEdited, EDependencyStrength::HardDependency));
+			TestEqual(TEXT("_5_SaveFoo2: Depends only on _4_CreateActor"), DependencyGraph.GetNodeById(Nodes[_5_SaveBar]).GetDependencies().Num(), 1);
+		}
+		// _6_RenameFooToBar
+		{
+			TestTrue(TEXT("_6_RenameFooToBar: Depends on _3_SaveFoo"), DependencyGraph.GetNodeById(Nodes[_6_RenameFooToBar]).DependsOnNode(Nodes[_3_SaveFoo], EActivityDependencyReason::PackageCreation, EDependencyStrength::HardDependency));
+			TestTrue(TEXT("_6_RenameFooToBar: Depends on _5_SaveBar"), DependencyGraph.GetNodeById(Nodes[_6_RenameFooToBar]).DependsOnNode(Nodes[_5_SaveBar], EActivityDependencyReason::PackageCreation, EDependencyStrength::HardDependency));
+			TestEqual(TEXT("_6_RenameFooToBar: Depends only on _3_SaveFoo and _5_SaveBar"), DependencyGraph.GetNodeById(Nodes[_6_RenameFooToBar]).GetDependencies().Num(), 2);
+		}
+		// _7_NewMapFoo
+		{
+			TestEqual(TEXT("_7_NewMapFoo: Depends on nothing"), DependencyGraph.GetNodeById(Nodes[_7_NewMapFoo]).GetDependencies().Num(), 0);
+		}
+		// _8_SaveMapFoo
+		{
+			TestTrue(TEXT("_8_SaveMapFoo: Depends on _7_NewMapBar"), DependencyGraph.GetNodeById(Nodes[_8_SaveMapFoo]).DependsOnNode(Nodes[_7_NewMapFoo], EActivityDependencyReason::PackageCreation, EDependencyStrength::HardDependency));
+			TestEqual(TEXT("_8_SaveMapFoo: Depends only on _7_NewMapBar"), DependencyGraph.GetNodeById(Nodes[_8_SaveMapFoo]).GetDependencies().Num(), 1);
+		}
+		// _9_CreateActor
+		{
+			TestTrue(TEXT("_9_CreateActor: Depends on _7_NewMapBar"), DependencyGraph.GetNodeById(Nodes[_9_CreateActor]).DependsOnNode(Nodes[_7_NewMapFoo], EActivityDependencyReason::PackageCreation, EDependencyStrength::HardDependency));
+			TestEqual(TEXT("_9_CreateActor: Depends only on _7_NewMapBar"), DependencyGraph.GetNodeById(Nodes[_9_CreateActor]).GetDependencies().Num(), 1);
+		}
+		// _10_SaveFoo
+		{
+			TestTrue(TEXT("_10_SaveFoo: Depends on _7_NewMapFoo"), DependencyGraph.GetNodeById(Nodes[_10_SaveFoo]).DependsOnNode(Nodes[_7_NewMapFoo], EActivityDependencyReason::PackageCreation, EDependencyStrength::HardDependency));
+			TestTrue(TEXT("_10_SaveFoo: Depends on _9_CreateActor"), DependencyGraph.GetNodeById(Nodes[_10_SaveFoo]).DependsOnNode(Nodes[_9_CreateActor], EActivityDependencyReason::PackageEdited, EDependencyStrength::HardDependency));
+			TestEqual(TEXT("_10_SaveFoo: Depends only on _7_NewMapFoo and _9_CreateActor"), DependencyGraph.GetNodeById(Nodes[_10_SaveFoo]).GetDependencies().Num(), 2);
+		}
+		// _11_DeleteFoo
+		{
+			TestTrue(TEXT("_11_DeleteFoo: Depends on _7_NewMapFoo"), DependencyGraph.GetNodeById(Nodes[_11_DeleteFoo]).DependsOnNode(Nodes[_7_NewMapFoo], EActivityDependencyReason::PackageCreation, EDependencyStrength::HardDependency));
+			TestEqual(TEXT("_11_DeleteFoo: Depends only on _7_NewMapFoo"), DependencyGraph.GetNodeById(Nodes[_11_DeleteFoo]).GetDependencies().Num(), 1);
+		}
+		// _12_NewMapFoo
+		{
+			TestTrue(TEXT("_12_NewMapFoo: Depends on _11_DeleteFoo"), DependencyGraph.GetNodeById(Nodes[_12_NewMapFoo]).DependsOnNode(Nodes[_11_DeleteFoo], EActivityDependencyReason::PackageRemoval, EDependencyStrength::HardDependency));
+			TestEqual(TEXT("_12_NewMapFoo: Depends only on _11_DeleteFoo"), DependencyGraph.GetNodeById(Nodes[_12_NewMapFoo]).GetDependencies().Num(), 1);
+		}
+		// _13_SaveMapFoo
+		{
+			TestTrue(TEXT("_13_SaveMapFoo: Depends on _12_NewMapFoo"), DependencyGraph.GetNodeById(Nodes[_13_SaveMapFoo]).DependsOnNode(Nodes[_12_NewMapFoo], EActivityDependencyReason::PackageCreation, EDependencyStrength::HardDependency));
+			TestEqual(TEXT("_13_SaveMapFoo: Depends only on _12_NewMapFoo"), DependencyGraph.GetNodeById(Nodes[_13_SaveMapFoo]).GetDependencies().Num(), 1);
+		}
+		// _14_CreateActor
+		{
+			TestTrue(TEXT("_14_CreateActor: Depends on _12_NewMapFoo"), DependencyGraph.GetNodeById(Nodes[_14_CreateActor]).DependsOnNode(Nodes[_12_NewMapFoo], EActivityDependencyReason::PackageCreation, EDependencyStrength::HardDependency));
+			TestEqual(TEXT("_14_CreateActor: Depends only on _12_NewMapFoo"), DependencyGraph.GetNodeById(Nodes[_14_CreateActor]).GetDependencies().Num(), 1);
+		}
+		// _15_SaveFoo
+		{
+			TestTrue(TEXT("_15_SaveFoo: Depends on _12_NewMapFoo"), DependencyGraph.GetNodeById(Nodes[_15_SaveFoo]).DependsOnNode(Nodes[_12_NewMapFoo], EActivityDependencyReason::PackageCreation, EDependencyStrength::HardDependency));
+			TestTrue(TEXT("_15_SaveFoo: Depends on _14_CreateActor"), DependencyGraph.GetNodeById(Nodes[_15_SaveFoo]).DependsOnNode(Nodes[_14_CreateActor], EActivityDependencyReason::PackageEdited, EDependencyStrength::HardDependency));
+			TestEqual(TEXT("_15_SaveFoo: Depends only on _12_NewMapFoo and _14_CreateActor"), DependencyGraph.GetNodeById(Nodes[_15_SaveFoo]).GetDependencies().Num(), 2);
+		}
+		
+		return true;
+	}
+
+	TArray<FActivityID> FillDatabase(FScopedSessionDatabase& SessionDatabase)
+	{
+		const FName FooLevel = TEXT("/Game/Foo");
+		const FName BarLevel = TEXT("/Game/Bar");
+		FTestActivityBuilder ActivityBuilder(SessionDatabase, ActivityCount);
+		
+		ActivityBuilder.CreateActor(FooLevel, _1_CreateActor, TEXT("FirstActor"));
+		ActivityBuilder.SaveMap(FooLevel, _2_SaveFoo);
+		ActivityBuilder.SaveMap(FooLevel, _3_SaveFoo);
+		ActivityBuilder.CreateActor(FooLevel, _4_CreateActor, TEXT("SecondActor"));
+		ActivityBuilder.RenameMap(FooLevel, BarLevel, _5_SaveBar, _6_RenameFooToBar);
+		ActivityBuilder.NewMap(FooLevel, _7_NewMapFoo);
+		ActivityBuilder.SaveMap(FooLevel, _8_SaveMapFoo);
+		ActivityBuilder.CreateActor(FooLevel, _9_CreateActor, TEXT("ThirdActor"));
+		ActivityBuilder.SaveMap(FooLevel, _10_SaveFoo);
+		ActivityBuilder.DeleteMap(FooLevel, _11_DeleteFoo);
+		ActivityBuilder.NewMap(FooLevel, _12_NewMapFoo);
+		ActivityBuilder.SaveMap(FooLevel, _13_SaveMapFoo);
+		ActivityBuilder.CreateActor(FooLevel, _14_CreateActor);
+		ActivityBuilder.SaveMap(FooLevel, _15_SaveFoo);
+		
+		return ActivityBuilder.GetActivities();
+	}
+	
+	FConcertExportedObject CreateActorMetaData(FName OuterLevelPath, FName ActorName)
+	{
+		FConcertExportedObject Result;
+		Result.ObjectId.ObjectName = ActorName;
+		Result.ObjectId.ObjectPackageName = OuterLevelPath;
+		Result.ObjectId.ObjectOuterPathName = *FString::Printf(TEXT("%s:PersistentLevel"), *OuterLevelPath.ToString());
+		Result.ObjectId.ObjectClassPathName = TEXT("/Script/Engine.StaticMeshActor");
+		return Result;
+	}
+}
+
+namespace UE::ConcertSyncTests::RenamingDependencyEdgeCases
+{
+	enum ETestActivity
+	{
+		_1_AddFoo,
+		_2_AddBar,
+		_3_DeleteFoo,
+		_4_DeleteBar,
+
+		// The goal is to test that _7_RenameFooToBar depends on _5_AddFoo and _6_SaveBar (instead of _1_AddFoo nor _2_AddBar)
+		_5_AddFoo,
+		_6_SaveBar,
+		_7_RenameFooToBar,
+
+		// The goal is to test that _11_RenameBarToFoo depends on _9_AddBar and _10_SaveFoo (instead of _1_AddFoo, _5_AddFoo nor _2_AddBar, _7_RenameFooToBar
+		_8_SaveFoo,
+		_9_RenameBarToFoo,
+
+		ActivityCount
+	};
+	
+	TArray<FActivityID> FillDatabase(FScopedSessionDatabase& SessionDatabase);
+	
+	/**
+	 * This tests that renaming dependencies are correctly set in some edge cases.
+	 *
+	 * User actions:
+	 *  1 Add Foo
+	 *  2 Add Bar
+	 *  3 Delete Foo
+	 *  4 Delete Bar
+	 *  5 Add Foo
+	 *  6 Rename Foo to Bar
+	 *  7 Delete Bar
+	 *  8 Add Bar
+	 *  9 Rename Bar to Foo
+	 */
+	IMPLEMENT_SIMPLE_AUTOMATION_TEST(FRenamingDependencyEdgeCaseTests, "Concert.History.BuildGraph.RenamingDependencyEdgeCaseTests", EAutomationTestFlags::ApplicationContextMask | EAutomationTestFlags::EngineFilter)
+	bool FRenamingDependencyEdgeCaseTests::RunTest(const FString& Parameters)
+	{
+		using namespace ConcertSyncCore;
+		
+		FScopedSessionDatabase SessionDatabase(*this);
+		TArray<FActivityID> Activities = FillDatabase(SessionDatabase);
+		
+		const FActivityDependencyGraph DependencyGraph = BuildDependencyGraphFrom(SessionDatabase);
+		UE_LOG(LogConcert, Log, TEXT("%s tested graph in Graphviz format:\n\n%s"), *GetTestFullName(), *ConcertSyncCore::Graphviz::ExportToGraphviz(DependencyGraph, SessionDatabase));
+		TArray<FActivityNodeID> Nodes = ConcertSyncTests::ValidateEachActivityHasNode(*this, Activities, DependencyGraph, ActivityCount, [](const uint32 ActivityType){ return FString::FromInt(ActivityType); });
+
+		// _7_RenameFooToBar
+		{
+			TestTrue(TEXT("_7_RenameFooToBar > _5_AddFoo"), DependencyGraph.GetNodeById(Nodes[_7_RenameFooToBar]).DependsOnNode(Nodes[_5_AddFoo], EActivityDependencyReason::PackageCreation, EDependencyStrength::HardDependency));
+			TestTrue(TEXT("_7_RenameFooToBar > _6_SaveBar"), DependencyGraph.GetNodeById(Nodes[_7_RenameFooToBar]).DependsOnNode(Nodes[_6_SaveBar], EActivityDependencyReason::PackageCreation, EDependencyStrength::HardDependency));
+			TestEqual(TEXT("_7_RenameFooToBar only depends on _5_AddFoo and _6_SaveBar"), DependencyGraph.GetNodeById(Nodes[_7_RenameFooToBar]).GetDependencies().Num(), 2);
+		}
+
+		// _9_RenameBarToFoo
+		{
+			TestTrue(TEXT("_9_RenameBarToFoo > _7_RenameFooToBar"), DependencyGraph.GetNodeById(Nodes[_9_RenameBarToFoo]).DependsOnNode(Nodes[_7_RenameFooToBar], EActivityDependencyReason::PackageRename, EDependencyStrength::HardDependency));
+			TestTrue(TEXT("_9_RenameBarToFoo > _8_SaveFoo"), DependencyGraph.GetNodeById(Nodes[_9_RenameBarToFoo]).DependsOnNode(Nodes[_8_SaveFoo], EActivityDependencyReason::PackageCreation, EDependencyStrength::HardDependency));
+			TestEqual(TEXT("_9_RenameBarToFoo only depends on _7_RenameFooToBar and _8_SaveFoo"), DependencyGraph.GetNodeById(Nodes[_9_RenameBarToFoo]).GetDependencies().Num(), 2);
+		}
+		
+		return true;
+	}
+	
+	TArray<FActivityID> FillDatabase(FScopedSessionDatabase& SessionDatabase)
+	{
+		const FName FooLevel = TEXT("/Game/Foo");
+		const FName BarLevel = TEXT("/Game/Bar");
+		FTestActivityBuilder ActivityBuilder(SessionDatabase, ActivityCount);
+
+		ActivityBuilder.NewMap(FooLevel, _1_AddFoo);
+		ActivityBuilder.NewMap(BarLevel, _2_AddBar);
+		ActivityBuilder.DeleteMap(FooLevel, _3_DeleteFoo);
+		ActivityBuilder.DeleteMap(BarLevel, _4_DeleteBar);
+		ActivityBuilder.NewMap(FooLevel, _5_AddFoo);
+		ActivityBuilder.RenameMap(FooLevel, BarLevel, _6_SaveBar, _7_RenameFooToBar);
+		ActivityBuilder.RenameMap(BarLevel, FooLevel, _8_SaveFoo, _9_RenameBarToFoo);
+		
+		return ActivityBuilder.GetActivities();
+	}
+}
+
+
+namespace UE::ConcertSyncTests::TransactionDependencyEdgeCases
+{
+	enum ETestActivity
+	{
+		_1_AddFoo,
+		_2_CreateActor,
+		_3_DeleteFoo,
+
+		// The goal is to test _5_CreateActor does not depend on _2_CreateActor
+		_4_AddFoo,
+		_5_CreateActor,
+		_6_SaveBar,
+		_7_RenameFooToBar,
+
+		// The goal is to that _9_CreateActor does not depend on _5_CreateActor
+		_8_AddFoo,
+		_9_CreateActor,
+
+		ActivityCount
+	};
+	
+	TArray<FActivityID> FillDatabase(FScopedSessionDatabase& SessionDatabase);
+	
+	IMPLEMENT_SIMPLE_AUTOMATION_TEST(FTransactionDependencyEdgeCaseTests, "Concert.History.BuildGraph.TransactionDependencyEdgeCaseTests", EAutomationTestFlags::ApplicationContextMask | EAutomationTestFlags::EngineFilter)
+	bool FTransactionDependencyEdgeCaseTests::RunTest(const FString& Parameters)
+	{
+		using namespace ConcertSyncCore;
+		
+		FScopedSessionDatabase SessionDatabase(*this);
+		TArray<FActivityID> Activities = FillDatabase(SessionDatabase);
+		
+		const FActivityDependencyGraph DependencyGraph = BuildDependencyGraphFrom(SessionDatabase);
+		UE_LOG(LogConcert, Log, TEXT("%s tested graph in Graphviz format:\n\n%s"), *GetTestFullName(), *ConcertSyncCore::Graphviz::ExportToGraphviz(DependencyGraph, SessionDatabase));
+		TArray<FActivityNodeID> Nodes = ConcertSyncTests::ValidateEachActivityHasNode(*this, Activities, DependencyGraph, ActivityCount, [](const uint32 ActivityType){ return FString::FromInt(ActivityType); });
+
+		// The main take away from this test:
+		TestFalse(TEXT("_5_CreateActor does not depend on _2_CreateActor"), DependencyGraph.GetNodeById(Nodes[_2_CreateActor]).AffectsNode(Nodes[_5_CreateActor]));
+		TestFalse(TEXT("_9_CreateActor does not depend on _5_CreateActor"), DependencyGraph.GetNodeById(Nodes[_5_CreateActor]).AffectsNode(Nodes[_9_CreateActor]));
+
+		// For completion sake:
+		TestTrue(TEXT("_5_CreateActor depends on _4_AddFoo"), DependencyGraph.GetNodeById(Nodes[_5_CreateActor]).DependsOnNode(Nodes[_4_AddFoo]));
+		TestEqual(TEXT("_5_CreateActor only depends on _4_AddFoo"), DependencyGraph.GetNodeById(Nodes[_5_CreateActor]).GetDependencies().Num(), 1);
+		TestTrue(TEXT("_9_CreateActor depends on _8_AddFoo"), DependencyGraph.GetNodeById(Nodes[_9_CreateActor]).DependsOnNode(Nodes[_8_AddFoo]));
+		TestEqual(TEXT("_9_CreateActor only depends on _8_AddFoo"), DependencyGraph.GetNodeById(Nodes[_9_CreateActor]).GetDependencies().Num(), 1);
+
+		return true;
+	}
+
+	TArray<FActivityID> FillDatabase(FScopedSessionDatabase& SessionDatabase)
+	{
+		const FName FooLevel = TEXT("/Game/Foo");
+		const FName BarLevel = TEXT("/Game/Bar");
+		FTestActivityBuilder ActivityBuilder(SessionDatabase, ActivityCount);
+
+		ActivityBuilder.NewMap(FooLevel, _1_AddFoo);
+		ActivityBuilder.CreateActor(FooLevel, _2_CreateActor, TEXT("FirstActor"));
+		ActivityBuilder.DeleteMap(FooLevel, _3_DeleteFoo);
+		ActivityBuilder.NewMap(FooLevel, _4_AddFoo);
+		ActivityBuilder.CreateActor(FooLevel, _5_CreateActor, TEXT("SecondActor"));
+		ActivityBuilder.RenameMap(FooLevel, BarLevel, _6_SaveBar, _7_RenameFooToBar);
+		ActivityBuilder.NewMap(FooLevel, _8_AddFoo);
+		ActivityBuilder.CreateActor(FooLevel, _9_CreateActor, TEXT("ThirdActor"));
+		
+		return ActivityBuilder.GetActivities();
+	}
 }
\ No newline at end of file
diff --git a/Engine/Plugins/Developer/Concert/ConcertSync/ConcertSyncTest/Source/ConcertSyncTest/Private/HistoryEdition/HistoryBottomUpAnalysisTests.cpp b/Engine/Plugins/Developer/Concert/ConcertSync/ConcertSyncTest/Source/ConcertSyncTest/Private/HistoryEdition/HistoryBottomUpAnalysisTests.cpp
--- a/Engine/Plugins/Developer/Concert/ConcertSync/ConcertSyncTest/Source/ConcertSyncTest/Private/HistoryEdition/HistoryBottomUpAnalysisTests.cpp
+++ b/Engine/Plugins/Developer/Concert/ConcertSync/ConcertSyncTest/Source/ConcertSyncTest/Private/HistoryEdition/HistoryBottomUpAnalysisTests.cpp
@@ -35,21 +35,24 @@ namespace UE::ConcertSyncTests::AnalysisTests
 		// I highly recommend visualising the above graph to understand the test better
 		// Unmute an edit operation that took place after an actor rename
 		{
-			const FHistoryAnalysisResult UnmuteEditAfterRename = AnalyseActivityDependencies_BottomUp({ Activities[_4_EditActor] }, DependencyGraph);
+			const FHistoryAnalysisResult UnmuteEditAfterRename = AnalyseActivityDependencies_BottomUp({ Activities[_5_EditActor] }, DependencyGraph);
 			TestEqual(TEXT("UnmuteEditAfterRename.HardDependencies.Num() == 2"), UnmuteEditAfterRename.HardDependencies.Num(), 2);
 			TestEqual(TEXT("UnmuteEditAfterRename.PossibleDependencies.Num() == 1"), UnmuteEditAfterRename.PossibleDependencies.Num(), 1);
-			TestTrue(TEXT("Editing actor depends on creating actor"), UnmuteEditAfterRename.HardDependencies.Contains(Activities[_2_AddActor]));
+			TestTrue(TEXT("Editing actor depends on creating actor"), UnmuteEditAfterRename.HardDependencies.Contains(Activities[_3_AddActor]));
 			TestTrue(TEXT("Creating actor depends on creating map"), UnmuteEditAfterRename.HardDependencies.Contains(Activities[_1_NewPackageFoo]));
-			TestTrue(TEXT("Editting possibly depends on renaming"), UnmuteEditAfterRename.PossibleDependencies.Contains(Activities[_3_RenameActor]));
+			TestTrue(TEXT("Editting possibly depends on renaming"), UnmuteEditAfterRename.PossibleDependencies.Contains(Activities[_4_RenameActor]));
 		}
 
 		// Unmute package rename
 		{
-			const FHistoryAnalysisResult UnmutePackageRename = AnalyseActivityDependencies_BottomUp({ Activities[_5_RenameFooToBar] }, DependencyGraph);
-			TestEqual(TEXT("UnmuteEditAfterRename.HardDependencies.Num() == 1"), UnmutePackageRename.HardDependencies.Num(), 1);
-			TestEqual(TEXT("UnmuteEditAfterRename.PossibleDependencies.Num() == 1"), UnmutePackageRename.PossibleDependencies.Num(), 1);
-			TestTrue(TEXT("Renaming package depends on having create the package"), UnmutePackageRename.HardDependencies.Contains(Activities[_1_NewPackageFoo]));
-			TestTrue(TEXT("Renaming package possibly depends on saving package"), UnmutePackageRename.PossibleDependencies.Contains(Activities[_5_SavePackageBar]));
+			const FHistoryAnalysisResult UnmutePackageRename = AnalyseActivityDependencies_BottomUp({ Activities[_7_RenameFooToBar] }, DependencyGraph);
+			TestEqual(TEXT("UnmutePackageRename.HardDependencies.Num() == 5"), UnmutePackageRename.HardDependencies.Num(), 5);
+			TestEqual(TEXT("UnmutePackageRename.PossibleDependencies.Num() == 0"), UnmutePackageRename.PossibleDependencies.Num(), 0);
+			TestTrue(TEXT("_7_RenameFooToBar > _1_NewPackageFoo"), UnmutePackageRename.HardDependencies.Contains(Activities[_1_NewPackageFoo]));
+			TestTrue(TEXT("_7_RenameFooToBar > _3_AddActor"), UnmutePackageRename.HardDependencies.Contains(Activities[_3_AddActor]));
+			TestTrue(TEXT("_7_RenameFooToBar > _4_RenameActor"), UnmutePackageRename.HardDependencies.Contains(Activities[_4_RenameActor]));
+			TestTrue(TEXT("_7_RenameFooToBar > _5_EditActor"), UnmutePackageRename.HardDependencies.Contains(Activities[_5_EditActor]));
+			TestTrue(TEXT("_7_RenameFooToBar > _6_SavePackageBar"), UnmutePackageRename.HardDependencies.Contains(Activities[_6_SavePackageBar]));
 		}
 		
 		return true;
diff --git a/Engine/Plugins/Developer/Concert/ConcertSync/ConcertSyncTest/Source/ConcertSyncTest/Private/HistoryEdition/HistoryTestUtil.h b/Engine/Plugins/Developer/Concert/ConcertSync/ConcertSyncTest/Source/ConcertSyncTest/Private/HistoryEdition/HistoryTestUtil.h
new file mode 100644
--- /dev/null
+++ b/Engine/Plugins/Developer/Concert/ConcertSync/ConcertSyncTest/Source/ConcertSyncTest/Private/HistoryEdition/HistoryTestUtil.h
@@ -0,0 +1,144 @@
+// Copyright Epic Games, Inc. All Rights Reserved.
+
+#pragma once
+
+#include "CoreMinimal.h"
+#include "ConcertMessageData.h"
+#include "ScopedSessionDatabase.h"
+#include "HistoryEdition/ActivityGraphIDs.h"
+#include "HistoryEdition/ActivityNode.h"
+#include "Misc/AutomationTest.h"
+
+namespace UE::ConcertSyncTests
+{
+	TArray<ConcertSyncCore::FActivityNodeID> ValidateEachActivityHasNode(
+		FAutomationTestBase& Test,
+		const TArray<int64>& ActivityMappings,
+		const ConcertSyncCore::FActivityDependencyGraph& Graph,
+		uint32 ActivityCount,
+		TFunctionRef<FString(uint32 ActivityType)> LexToString)
+	{
+		using namespace ConcertSyncCore;
+		
+		TArray<FActivityNodeID> ActivityNodes;
+		ActivityNodes.SetNumZeroed(ActivityCount);
+		for (uint32 ActivityType = 0; ActivityType < ActivityCount; ++ActivityType)
+		{
+			const int64 ActivityId = ActivityMappings[ActivityType];
+			const TOptional<FActivityNodeID> NodeID = Graph.FindNodeByActivity(ActivityId);
+			if (!NodeID.IsSet())
+			{
+				Test.AddError(FString::Printf(TEXT("No node generated for activity type %s"), *LexToString(ActivityType)));
+				continue;
+			}
+			
+			if (!NodeID.IsSet())
+			{
+				Test.AddError(FString::Printf(TEXT("Graph has invalid state. Node ID %lld is invalid for activity type %s"), NodeID->ID, *LexToString(ActivityType)));
+				continue;
+			}
+			
+			ActivityNodes[ActivityType] = *NodeID;
+		}
+
+		return ActivityNodes;
+	}
+
+	/**
+	 * Helps fill the database
+	 */
+	class FTestActivityBuilder
+	{
+		TArray<FActivityID> Activities;
+		FScopedSessionDatabase& SessionDatabase;
+	public:
+
+		FTestActivityBuilder(FScopedSessionDatabase& SessionDatabase, uint32 ActivityCount)
+			: SessionDatabase(SessionDatabase)
+		{
+			Activities.SetNumUninitialized(ActivityCount);
+		}
+
+		TArray<FActivityID> GetActivities() const { return Activities; }
+
+		bool NewMap(FName MapName, uint32 ActivityIndex)
+		{
+			FConcertSyncActivity NewPackage;
+			NewPackage.EndpointId = SessionDatabase.GetEndpoint();
+			FConcertPackageInfo PackageInfo;
+			FConcertPackageDataStream PackageDataStream;
+			PackageInfo.PackageName = MapName;
+			PackageInfo.PackageUpdateType = EConcertPackageUpdateType::Added;
+			SessionDatabase.GetTransactionMaxEventId(PackageInfo.TransactionEventIdAtSave);
+			int64 Dummy;
+			return SessionDatabase.AddPackageActivity(NewPackage, PackageInfo, PackageDataStream, Activities[ActivityIndex], Dummy);
+		}
+
+		bool SaveMap(FName MapName, uint32 ActivityIndex)
+		{
+			FConcertSyncActivity SaveFooPackage;
+			SaveFooPackage.EndpointId = SessionDatabase.GetEndpoint();
+			FConcertPackageInfo PackageInfo;
+			FConcertPackageDataStream PackageDataStream;
+			PackageInfo.PackageName = MapName;
+			PackageInfo.PackageUpdateType = EConcertPackageUpdateType::Saved;
+			SessionDatabase.GetTransactionMaxEventId(PackageInfo.TransactionEventIdAtSave);
+			int64 Dummy;
+			return SessionDatabase.AddPackageActivity(SaveFooPackage, PackageInfo, PackageDataStream, Activities[ActivityIndex], Dummy);
+		}
+
+		bool DeleteMap(FName MapName, uint32 ActivityIndex)
+		{
+			FConcertSyncActivity DeleteFooPackage;
+			DeleteFooPackage.EndpointId = SessionDatabase.GetEndpoint();
+			FConcertPackageInfo PackageInfo;
+			FConcertPackageDataStream PackageDataStream;
+			PackageInfo.PackageName = MapName;
+			PackageInfo.PackageUpdateType = EConcertPackageUpdateType::Deleted;
+			SessionDatabase.GetTransactionMaxEventId(PackageInfo.TransactionEventIdAtSave);
+			int64 Dummy;
+			return SessionDatabase.AddPackageActivity(DeleteFooPackage, PackageInfo, PackageDataStream, Activities[ActivityIndex], Dummy);
+		}
+
+		bool RenameMap(FName OldMapName, FName NewMapName, uint32 SaveActivityIndex, uint32 RenameActivityIndex)
+		{
+			int64 Dummy;
+			
+			FConcertSyncActivity SaveFoo2Package;
+			SaveFoo2Package.EndpointId = SessionDatabase.GetEndpoint();
+			FConcertPackageInfo PackageInfo;
+			FConcertPackageDataStream PackageDataStream;
+			PackageInfo.PackageName = NewMapName;
+			PackageInfo.PackageUpdateType = EConcertPackageUpdateType::Saved;
+			SessionDatabase.GetTransactionMaxEventId(PackageInfo.TransactionEventIdAtSave);
+			bool bSuccess = SessionDatabase.AddPackageActivity(SaveFoo2Package, PackageInfo, PackageDataStream, Activities[SaveActivityIndex], Dummy);
+			
+			PackageInfo.PackageUpdateType = EConcertPackageUpdateType::Renamed;
+			PackageInfo.PackageName = OldMapName;
+			PackageInfo.NewPackageName = NewMapName;
+			bSuccess &= SessionDatabase.AddPackageActivity(SaveFoo2Package, PackageInfo, PackageDataStream, Activities[RenameActivityIndex], Dummy);
+
+			return bSuccess;
+		}
+
+		bool CreateActor(FName MapName, uint32 ActivityIndex, FName ActorName = EName::Actor)
+		{
+			FConcertExportedObject Actor;
+			Actor.ObjectId.ObjectName = ActorName;
+			Actor.ObjectId.ObjectPackageName = MapName;
+			Actor.ObjectId.ObjectOuterPathName = *FString::Printf(TEXT("%s:PersistentLevel"), *MapName.ToString());
+			Actor.ObjectId.ObjectClassPathName = TEXT("/Script/Engine.StaticMeshActor");
+			
+			FConcertSyncTransactionActivity CreateActorActivity;
+			CreateActorActivity.EndpointId = SessionDatabase.GetEndpoint();
+			CreateActorActivity.EventData.Transaction.TransactionId = FGuid::NewGuid();
+			CreateActorActivity.EventData.Transaction.OperationId = FGuid::NewGuid();
+			Actor.ObjectData.bAllowCreate = true;
+			CreateActorActivity.EventData.Transaction.ExportedObjects = { Actor };
+			CreateActorActivity.EventData.Transaction.ModifiedPackages = { MapName };
+			SessionDatabase.GetTransactionMaxEventId(CreateActorActivity.EventId);
+			int64 Dummy;
+			return SessionDatabase.AddTransactionActivity(CreateActorActivity, Activities[ActivityIndex], Dummy);
+		}
+	};
+}
diff --git a/Engine/Plugins/Developer/Concert/ConcertSync/ConcertSyncTest/Source/ConcertSyncTest/Private/HistoryEdition/HistoryTopDownAnalysisTests.cpp b/Engine/Plugins/Developer/Concert/ConcertSync/ConcertSyncTest/Source/ConcertSyncTest/Private/HistoryEdition/HistoryTopDownAnalysisTests.cpp
--- a/Engine/Plugins/Developer/Concert/ConcertSync/ConcertSyncTest/Source/ConcertSyncTest/Private/HistoryEdition/HistoryTopDownAnalysisTests.cpp
+++ b/Engine/Plugins/Developer/Concert/ConcertSync/ConcertSyncTest/Source/ConcertSyncTest/Private/HistoryEdition/HistoryTopDownAnalysisTests.cpp
@@ -74,15 +74,9 @@ namespace UE::ConcertSyncTests::AnalysisTests
 		// Delete /Game/Foo > Nearly everything has hard dependency
 		{
 			const FHistoryAnalysisResult DeleteFooRequirements = AnalyseActivityDependencies_TopDown({ Activities[_1_NewPackageFoo] }, DependencyGraph);
-			/* _1_NewPackageFoo: is what we're "deleting".
-			 * _5_SavePackageBar: Bar is created as result of a rename but has no dependency to _1_NewPackageFoo.
-			 * All other activities transitively depend on _1_NewPackageFoo (put above log into GraphViz to visualise).
-			 * 
-			 * Note: The transaction activities (_3_EditActor, _4_EditActor) have possible dependencies BUT they do have hard dependencies to _2_AddActor. This is why they must be in HardDependencies, too.
-			 */
-			const TSet<ETestActivity> ExcludedHardActivities { _1_NewPackageFoo, _5_SavePackageBar };
+			const TSet<ETestActivity> ExcludedHardActivities { _1_NewPackageFoo };
 			const TSet<ETestActivity> HardDependencies = AllActivities().Difference(ExcludedHardActivities);
-			const TSet<ETestActivity> PossibleDependencies { _5_SavePackageBar };
+			const TSet<ETestActivity> PossibleDependencies { };
 			const bool bDeleteAllCorrect = ValidateRequirements(TEXT("Delete /Game/Foo"), *this, Activities, DeleteFooRequirements, HardDependencies, PossibleDependencies);
 			TestTrue(TEXT("Delete /Game/Foo is correct"), bDeleteAllCorrect);
 		}
@@ -90,36 +84,36 @@ namespace UE::ConcertSyncTests::AnalysisTests
 		// Delete rename transaction > No dependencies
 			// because the rename basically just changes AActor::ActorLabel
 		{
-			const FHistoryAnalysisResult DeleteRenameRequirements = AnalyseActivityDependencies_TopDown({ Activities[_3_RenameActor] }, DependencyGraph);
-			TestEqual(TEXT("Delete renaming actor: HardDependencies.Num() == 0"), DeleteRenameRequirements.HardDependencies.Num(), 0);
+			const FHistoryAnalysisResult DeleteRenameRequirements = AnalyseActivityDependencies_TopDown({ Activities[_4_RenameActor] }, DependencyGraph);
+			TestEqual(TEXT("Delete renaming actor: HardDependencies.Num() == 6"), DeleteRenameRequirements.HardDependencies.Num(), 6);
 			TestEqual(TEXT("Delete renaming actor: PossibleDependencies.Num() == 1"), DeleteRenameRequirements.PossibleDependencies.Num(), 1);
-			TestTrue(TEXT("Delete renaming actor: Edit activity may depend on deleted activity"), DeleteRenameRequirements.PossibleDependencies.Contains(Activities[_4_EditActor]));
+			TestTrue(TEXT("Delete renaming actor: Edit activity may depend on deleted activity"), DeleteRenameRequirements.PossibleDependencies.Contains(Activities[_5_EditActor]));
 		}
 
 		// Delete actor creation > All transactions operating on actor are hard dependencies
 		{
-			const FHistoryAnalysisResult DeleteCreateActorRequirements = AnalyseActivityDependencies_TopDown({ Activities[_2_AddActor] }, DependencyGraph);
-			TestEqual(TEXT("Delete actor creation: HardDependencies.Num() == 2"), DeleteCreateActorRequirements.HardDependencies.Num(), 2);
+			const FHistoryAnalysisResult DeleteCreateActorRequirements = AnalyseActivityDependencies_TopDown({ Activities[_3_AddActor] }, DependencyGraph);
+			TestEqual(TEXT("Delete actor creation: HardDependencies.Num() == 8"), DeleteCreateActorRequirements.HardDependencies.Num(), 8);
 			TestEqual(TEXT("Delete actor creation: PossibleDependencies.Num() == 0"), DeleteCreateActorRequirements.PossibleDependencies.Num(), 0);
-			TestTrue(TEXT("Delete actor creation: Rename depends on created actor"), DeleteCreateActorRequirements.HardDependencies.Contains(Activities[_3_RenameActor]));
-			TestTrue(TEXT("Delete actor creation: Edit depends on created actor"), DeleteCreateActorRequirements.HardDependencies.Contains(Activities[_4_EditActor]));
+			TestTrue(TEXT("Delete actor creation: Rename depends on created actor"), DeleteCreateActorRequirements.HardDependencies.Contains(Activities[_4_RenameActor]));
+			TestTrue(TEXT("Delete actor creation: Edit depends on created actor"), DeleteCreateActorRequirements.HardDependencies.Contains(Activities[_5_EditActor]));
 		}
 
 		// Deleting a rename activity > Rename activity should have a possible dependency to activities that saved the renamed-to package
 			// because MU internally creates a Save Package activity before a rename activity (so secretly rename activities is two related activities)
 		{
-			const FHistoryAnalysisResult DeleteRenameRequirements = AnalyseActivityDependencies_TopDown({ Activities[_5_RenameFooToBar] }, DependencyGraph);
+			const FHistoryAnalysisResult DeleteRenameRequirements = AnalyseActivityDependencies_TopDown({ Activities[_7_RenameFooToBar] }, DependencyGraph);
 
 			// Not really what we want to test - but we're doing it for completness
 			TestEqual(TEXT("Delete rename package: HardDependencies.Num() == 4"), DeleteRenameRequirements.HardDependencies.Num(), 4);
-			TestTrue(TEXT("Delete rename package: Rename has hard dependency creating renamed package"), DeleteRenameRequirements.HardDependencies.Contains(Activities[_6_EditActor]));
-			TestTrue(TEXT("Delete rename package: Rename has hard dependency to editing renamed package (1)"), DeleteRenameRequirements.HardDependencies.Contains(Activities[_7_DeleteBar]));
-			TestTrue(TEXT("Delete rename package: Rename has hard dependency to editing renamed package (2)"), DeleteRenameRequirements.HardDependencies.Contains(Activities[_8_NewPackageFoo]));
-			TestTrue(TEXT("Delete rename package: Rename has hard dependency to editing renamed package (3)"), DeleteRenameRequirements.HardDependencies.Contains(Activities[_8_SavePackageFoo]));
+			TestTrue(TEXT("Delete rename package: Rename has hard dependency creating renamed package"), DeleteRenameRequirements.HardDependencies.Contains(Activities[_8_EditActor]));
+			TestTrue(TEXT("Delete rename package: Rename has hard dependency to editing renamed package (1)"), DeleteRenameRequirements.HardDependencies.Contains(Activities[_9_DeleteBar]));
+			TestTrue(TEXT("Delete rename package: Rename has hard dependency to editing renamed package (2)"), DeleteRenameRequirements.HardDependencies.Contains(Activities[_10_NewPackageFoo]));
+			TestTrue(TEXT("Delete rename package: Rename has hard dependency to editing renamed package (3)"), DeleteRenameRequirements.HardDependencies.Contains(Activities[_11_SavePackageFoo]));
 
 			// Now what we really wanted to test
 			TestEqual(TEXT("Delete rename package: PossibleDependencies.Num() == 1"), DeleteRenameRequirements.PossibleDependencies.Num(), 1);
-			TestTrue(TEXT("Delete rename package: Rename has soft dependency to saving renamed-to package"), DeleteRenameRequirements.PossibleDependencies.Contains(Activities[_5_SavePackageBar]));
+			TestTrue(TEXT("Delete rename package: Rename has soft dependency to saving renamed-to package"), DeleteRenameRequirements.PossibleDependencies.Contains(Activities[_6_SavePackageBar]));
 		}
 
 		return true;
diff --git a/Engine/Plugins/Developer/Concert/ConcertSync/ConcertSyncTest/Source/ConcertSyncTest/Private/HistoryEdition/RenameEditAndDeleteMapsFlow.cpp b/Engine/Plugins/Developer/Concert/ConcertSync/ConcertSyncTest/Source/ConcertSyncTest/Private/HistoryEdition/RenameEditAndDeleteMapsFlow.cpp
--- a/Engine/Plugins/Developer/Concert/ConcertSync/ConcertSyncTest/Source/ConcertSyncTest/Private/HistoryEdition/RenameEditAndDeleteMapsFlow.cpp
+++ b/Engine/Plugins/Developer/Concert/ConcertSync/ConcertSyncTest/Source/ConcertSyncTest/Private/HistoryEdition/RenameEditAndDeleteMapsFlow.cpp
@@ -20,16 +20,16 @@ namespace UE::ConcertSyncTests::RenameEditAndDeleteMapsFlowTest
 	{
 		return {
 			_1_NewPackageFoo,
-			_1_SavePackageFoo,
-			_2_AddActor,
-			_3_RenameActor,
-			_4_EditActor,
-			_5_SavePackageBar,
-			_5_RenameFooToBar,
-			_6_EditActor,
-			_7_DeleteBar,
-			_8_NewPackageFoo,
-			_8_SavePackageFoo
+			_2_SavePackageFoo,
+			_3_AddActor,
+			_4_RenameActor,
+			_5_EditActor,
+			_6_SavePackageBar,
+			_7_RenameFooToBar,
+			_8_EditActor,
+			_9_DeleteBar,
+			_10_NewPackageFoo,
+			_11_SavePackageFoo
 		};
 	}
 
@@ -38,16 +38,16 @@ namespace UE::ConcertSyncTests::RenameEditAndDeleteMapsFlowTest
 		switch (Activity)
 		{
 		case _1_NewPackageFoo: return TEXT("\"1 New package Foo\"");
-		case _1_SavePackageFoo: return TEXT("\"1 Saved package Foo\"");
-		case _2_AddActor: return TEXT("\"2 Create actor\"");
-		case _3_RenameActor: return TEXT("\"3 Edit actor\"");
-		case _4_EditActor: return TEXT("\"4 Edit actor\"");
-		case _5_SavePackageBar: return TEXT("\"5 Save package\"");
-		case _5_RenameFooToBar: return TEXT("\"5 Rename Foo to Bar\"");
-		case _6_EditActor: return TEXT("\"6 Edit actor\"");
-		case _7_DeleteBar: return TEXT("\"7 Delete package Bar\"");
-		case _8_NewPackageFoo: return TEXT("\"8 Create package Bar\"");
-		case _8_SavePackageFoo: return TEXT("\"8 Save package Bar\"");
+		case _2_SavePackageFoo: return TEXT("\"1 Saved package Foo\"");
+		case _3_AddActor: return TEXT("\"2 Create actor\"");
+		case _4_RenameActor: return TEXT("\"3 Edit actor\"");
+		case _5_EditActor: return TEXT("\"4 Edit actor\"");
+		case _6_SavePackageBar: return TEXT("\"5 Save package\"");
+		case _7_RenameFooToBar: return TEXT("\"5 Rename Foo to Bar\"");
+		case _8_EditActor: return TEXT("\"6 Edit actor\"");
+		case _9_DeleteBar: return TEXT("\"7 Delete package Bar\"");
+		case _10_NewPackageFoo: return TEXT("\"8 Create package Bar\"");
+		case _11_SavePackageFoo: return TEXT("\"8 Save package Bar\"");
 				
 		case ActivityCount:
 		default:
@@ -91,7 +91,7 @@ namespace UE::ConcertSyncTests::RenameEditAndDeleteMapsFlowTest
 			FConcertSyncActivity SavePackage = MakeActivity<FConcertSyncActivity>(EndpointID);
 			PackageInfo.PackageUpdateType = EConcertPackageUpdateType::Saved;
 			SessionDatabase.GetTransactionMaxEventId(PackageInfo.TransactionEventIdAtSave);
-			bAllSucceeded &= SessionDatabase.AddPackageActivity(SavePackage, PackageInfo, PackageDataStream, ActivityIDs[_1_SavePackageFoo], PackageEventIDs[_1_SavePackageFoo]);
+			bAllSucceeded &= SessionDatabase.AddPackageActivity(SavePackage, PackageInfo, PackageDataStream, ActivityIDs[_2_SavePackageFoo], PackageEventIDs[_2_SavePackageFoo]);
 		}
 
 		// 2 Add actor A
@@ -105,7 +105,7 @@ namespace UE::ConcertSyncTests::RenameEditAndDeleteMapsFlowTest
 			CreateActor.EventData.Transaction.ExportedObjects = { NewActorData.StaticMeshComponent, NewActorData.Actor };
 			CreateActor.EventData.Transaction.ModifiedPackages = { FooLevel };
 			SessionDatabase.GetTransactionMaxEventId(CreateActor.EventId);
-			bAllSucceeded &= SessionDatabase.AddTransactionActivity(CreateActor, ActivityIDs[_2_AddActor], PackageEventIDs[_2_AddActor]);
+			bAllSucceeded &= SessionDatabase.AddTransactionActivity(CreateActor, ActivityIDs[_3_AddActor], PackageEventIDs[_3_AddActor]);
 		}
 
 		// 3 Rename actor A
@@ -119,7 +119,7 @@ namespace UE::ConcertSyncTests::RenameEditAndDeleteMapsFlowTest
 			EditActor.EventData.Transaction.ExportedObjects = { NewActorData.Actor };
 			EditActor.EventData.Transaction.ModifiedPackages = { FooLevel };
 			SessionDatabase.GetTransactionMaxEventId(EditActor.EventId);
-			bAllSucceeded &= SessionDatabase.AddTransactionActivity(EditActor, ActivityIDs[_3_RenameActor], PackageEventIDs[_3_RenameActor]);
+			bAllSucceeded &= SessionDatabase.AddTransactionActivity(EditActor, ActivityIDs[_4_RenameActor], PackageEventIDs[_4_RenameActor]);
 		}
 
 		// 4 Edit actor A
@@ -132,7 +132,7 @@ namespace UE::ConcertSyncTests::RenameEditAndDeleteMapsFlowTest
 			EditActor.EventData.Transaction.ExportedObjects = { NewActorData.StaticMeshComponent };
 			EditActor.EventData.Transaction.ModifiedPackages = { FooLevel };
 			SessionDatabase.GetTransactionMaxEventId(EditActor.EventId);
-			bAllSucceeded &= SessionDatabase.AddTransactionActivity(EditActor, ActivityIDs[_4_EditActor], PackageEventIDs[_4_EditActor]);
+			bAllSucceeded &= SessionDatabase.AddTransactionActivity(EditActor, ActivityIDs[_5_EditActor], PackageEventIDs[_5_EditActor]);
 		}
 
 		// 5 Rename map to Bar
@@ -143,12 +143,12 @@ namespace UE::ConcertSyncTests::RenameEditAndDeleteMapsFlowTest
 			PackageInfo.PackageName = BarLevel;
 			PackageInfo.PackageUpdateType = EConcertPackageUpdateType::Saved;
 			SessionDatabase.GetTransactionMaxEventId(PackageInfo.TransactionEventIdAtSave);
-			bAllSucceeded &= SessionDatabase.AddPackageActivity(SaveBarPackage, PackageInfo, PackageDataStream, ActivityIDs[_5_SavePackageBar], PackageEventIDs[_5_SavePackageBar]);
+			bAllSucceeded &= SessionDatabase.AddPackageActivity(SaveBarPackage, PackageInfo, PackageDataStream, ActivityIDs[_6_SavePackageBar], PackageEventIDs[_6_SavePackageBar]);
 
 			PackageInfo.PackageUpdateType = EConcertPackageUpdateType::Renamed;
 			PackageInfo.PackageName = FooLevel;
 			PackageInfo.NewPackageName = BarLevel;
-			bAllSucceeded &= SessionDatabase.AddPackageActivity(SaveBarPackage, PackageInfo, PackageDataStream, ActivityIDs[_5_RenameFooToBar], PackageEventIDs[_5_RenameFooToBar]);
+			bAllSucceeded &= SessionDatabase.AddPackageActivity(SaveBarPackage, PackageInfo, PackageDataStream, ActivityIDs[_7_RenameFooToBar], PackageEventIDs[_7_RenameFooToBar]);
 		}
 
 		// 6 Edit actor A
@@ -161,7 +161,7 @@ namespace UE::ConcertSyncTests::RenameEditAndDeleteMapsFlowTest
 			EditActor.EventData.Transaction.ExportedObjects = { NewActorData.StaticMeshComponent };
 			EditActor.EventData.Transaction.ModifiedPackages = { BarLevel };
 			SessionDatabase.GetTransactionMaxEventId(EditActor.EventId);
-			bAllSucceeded &= SessionDatabase.AddTransactionActivity(EditActor, ActivityIDs[_6_EditActor], PackageEventIDs[_6_EditActor]);
+			bAllSucceeded &= SessionDatabase.AddTransactionActivity(EditActor, ActivityIDs[_8_EditActor], PackageEventIDs[_8_EditActor]);
 		}
 
 		// 7 Delete map Bar
@@ -172,7 +172,7 @@ namespace UE::ConcertSyncTests::RenameEditAndDeleteMapsFlowTest
 			PackageInfo.PackageName = BarLevel;
 			PackageInfo.PackageUpdateType = EConcertPackageUpdateType::Deleted;
 			SessionDatabase.GetTransactionMaxEventId(PackageInfo.TransactionEventIdAtSave);
-			bAllSucceeded &= SessionDatabase.AddPackageActivity(SaveBarPackage, PackageInfo, PackageDataStream, ActivityIDs[_7_DeleteBar], PackageEventIDs[_7_DeleteBar]);
+			bAllSucceeded &= SessionDatabase.AddPackageActivity(SaveBarPackage, PackageInfo, PackageDataStream, ActivityIDs[_9_DeleteBar], PackageEventIDs[_9_DeleteBar]);
 		}
 		
 		// 8 Create map Bar
@@ -183,12 +183,12 @@ namespace UE::ConcertSyncTests::RenameEditAndDeleteMapsFlowTest
 			PackageInfo.PackageName = TEXT("/Game/Bar");
 			PackageInfo.PackageUpdateType = EConcertPackageUpdateType::Added;
 			SessionDatabase.GetTransactionMaxEventId(PackageInfo.TransactionEventIdAtSave);
-			bAllSucceeded &= SessionDatabase.AddPackageActivity(NewPackage, PackageInfo, PackageDataStream, ActivityIDs[_8_NewPackageFoo], PackageEventIDs[_8_NewPackageFoo]);
+			bAllSucceeded &= SessionDatabase.AddPackageActivity(NewPackage, PackageInfo, PackageDataStream, ActivityIDs[_10_NewPackageFoo], PackageEventIDs[_10_NewPackageFoo]);
 
 			FConcertSyncActivity SavePackage = MakeActivity<FConcertSyncActivity>(EndpointID);
 			PackageInfo.PackageUpdateType = EConcertPackageUpdateType::Saved;
 			SessionDatabase.GetTransactionMaxEventId(PackageInfo.TransactionEventIdAtSave);
-			bAllSucceeded &= SessionDatabase.AddPackageActivity(SavePackage, PackageInfo, PackageDataStream, ActivityIDs[_8_SavePackageFoo], PackageEventIDs[_8_SavePackageFoo]);
+			bAllSucceeded &= SessionDatabase.AddPackageActivity(SavePackage, PackageInfo, PackageDataStream, ActivityIDs[_11_SavePackageFoo], PackageEventIDs[_11_SavePackageFoo]);
 		}
 
 		if (!bAllSucceeded)
diff --git a/Engine/Plugins/Developer/Concert/ConcertSync/ConcertSyncTest/Source/ConcertSyncTest/Private/HistoryEdition/RenameEditAndDeleteMapsFlow.h b/Engine/Plugins/Developer/Concert/ConcertSync/ConcertSyncTest/Source/ConcertSyncTest/Private/HistoryEdition/RenameEditAndDeleteMapsFlow.h
--- a/Engine/Plugins/Developer/Concert/ConcertSync/ConcertSyncTest/Source/ConcertSyncTest/Private/HistoryEdition/RenameEditAndDeleteMapsFlow.h
+++ b/Engine/Plugins/Developer/Concert/ConcertSync/ConcertSyncTest/Source/ConcertSyncTest/Private/HistoryEdition/RenameEditAndDeleteMapsFlow.h
@@ -11,18 +11,17 @@ namespace UE::ConcertSyncTests::RenameEditAndDeleteMapsFlowTest
 {
 	enum ETestActivity
 	{
-		// Redundant = are visual finding aid during debugging
-		_1_NewPackageFoo	= 0,
-		_1_SavePackageFoo	= 1,
-		_2_AddActor			= 2,
-		_3_RenameActor		= 3,
-		_4_EditActor		= 4,
-		_5_SavePackageBar	= 5,
-		_5_RenameFooToBar	= 6,
-		_6_EditActor		= 7,
-		_7_DeleteBar		= 8,
-		_8_NewPackageFoo	= 9,
-		_8_SavePackageFoo	= 10,
+		_1_NewPackageFoo,
+		_2_SavePackageFoo,
+		_3_AddActor,
+		_4_RenameActor,
+		_5_EditActor,
+		_6_SavePackageBar,
+		_7_RenameFooToBar,
+		_8_EditActor,
+		_9_DeleteBar,
+		_10_NewPackageFoo,
+		_11_SavePackageFoo,
 		
 		ActivityCount
 	};
