Clean up tests a little

diff --git a/Engine/Plugins/Developer/Concert/ConcertSync/ConcertSyncTest/Source/ConcertSyncTest/Private/HistoryEdition/DependencyGraphBuilderTests.cpp b/Engine/Plugins/Developer/Concert/ConcertSync/ConcertSyncTest/Source/ConcertSyncTest/Private/HistoryEdition/DependencyGraphBuilderTests.cpp
--- a/Engine/Plugins/Developer/Concert/ConcertSync/ConcertSyncTest/Source/ConcertSyncTest/Private/HistoryEdition/DependencyGraphBuilderTests.cpp
+++ b/Engine/Plugins/Developer/Concert/ConcertSync/ConcertSyncTest/Source/ConcertSyncTest/Private/HistoryEdition/DependencyGraphBuilderTests.cpp
@@ -7,7 +7,8 @@
 #include "HistoryEdition/DependencyGraphBuilder.h"
 #include "HistoryTestUtil.h"
 #include "RenameEditAndDeleteMapsFlow.h"
-#include "ScopedSessionDatabase.h"
+#include "Util/ActivityBuilder.h"
+#include "Util/ScopedSessionDatabase.h"
 
 #include "Misc/AutomationTest.h"
 
@@ -346,7 +347,6 @@ namespace UE::ConcertSyncTests::PackageEditedDependencyTest
 	};
 	
 	TArray<FActivityID> FillDatabase(FScopedSessionDatabase& SessionDatabase);
-	FConcertExportedObject CreateActorMetaData(FName OuterLevelPath, FName ActorName);
 	
 	/**
 	 * This tests that EActivityDependencyReason::PackageEdited dependencies are discovered correctly.
@@ -469,7 +469,7 @@ namespace UE::ConcertSyncTests::PackageEditedDependencyTest
 	{
 		const FName FooLevel = TEXT("/Game/Foo");
 		const FName BarLevel = TEXT("/Game/Bar");
-		FTestActivityBuilder ActivityBuilder(SessionDatabase, ActivityCount);
+		FActivityBuilder ActivityBuilder(SessionDatabase, ActivityCount);
 		
 		ActivityBuilder.CreateActor(FooLevel, _1_CreateActor, TEXT("FirstActor"));
 		ActivityBuilder.SaveMap(FooLevel, _2_SaveFoo);
@@ -488,16 +488,6 @@ namespace UE::ConcertSyncTests::PackageEditedDependencyTest
 		
 		return ActivityBuilder.GetActivities();
 	}
-	
-	FConcertExportedObject CreateActorMetaData(FName OuterLevelPath, FName ActorName)
-	{
-		FConcertExportedObject Result;
-		Result.ObjectId.ObjectName = ActorName;
-		Result.ObjectId.ObjectPackageName = OuterLevelPath;
-		Result.ObjectId.ObjectOuterPathName = *FString::Printf(TEXT("%s:PersistentLevel"), *OuterLevelPath.ToString());
-		Result.ObjectId.ObjectClassPathName = TEXT("/Script/Engine.StaticMeshActor");
-		return Result;
-	}
 }
 
 namespace UE::ConcertSyncTests::RenamingDependencyEdgeCases
@@ -570,7 +560,7 @@ namespace UE::ConcertSyncTests::RenamingDependencyEdgeCases
 	{
 		const FName FooLevel = TEXT("/Game/Foo");
 		const FName BarLevel = TEXT("/Game/Bar");
-		FTestActivityBuilder ActivityBuilder(SessionDatabase, ActivityCount);
+		FActivityBuilder ActivityBuilder(SessionDatabase, ActivityCount);
 
 		ActivityBuilder.NewMap(FooLevel, _1_AddFoo);
 		ActivityBuilder.NewMap(BarLevel, _2_AddBar);
@@ -637,7 +627,7 @@ namespace UE::ConcertSyncTests::TransactionDependencyEdgeCases
 	{
 		const FName FooLevel = TEXT("/Game/Foo");
 		const FName BarLevel = TEXT("/Game/Bar");
-		FTestActivityBuilder ActivityBuilder(SessionDatabase, ActivityCount);
+		FActivityBuilder ActivityBuilder(SessionDatabase, ActivityCount);
 
 		ActivityBuilder.NewMap(FooLevel, _1_AddFoo);
 		ActivityBuilder.CreateActor(FooLevel, _2_CreateActor, TEXT("FirstActor"));
diff --git a/Engine/Plugins/Developer/Concert/ConcertSync/ConcertSyncTest/Source/ConcertSyncTest/Private/HistoryEdition/HistoryBottomUpAnalysisTests.cpp b/Engine/Plugins/Developer/Concert/ConcertSync/ConcertSyncTest/Source/ConcertSyncTest/Private/HistoryEdition/HistoryBottomUpAnalysisTests.cpp
--- a/Engine/Plugins/Developer/Concert/ConcertSync/ConcertSyncTest/Source/ConcertSyncTest/Private/HistoryEdition/HistoryBottomUpAnalysisTests.cpp
+++ b/Engine/Plugins/Developer/Concert/ConcertSync/ConcertSyncTest/Source/ConcertSyncTest/Private/HistoryEdition/HistoryBottomUpAnalysisTests.cpp
@@ -3,17 +3,14 @@
 #include "ConcertLogGlobal.h"
 
 #include "ConcertSyncSessionDatabase.h"
+#include "HistoryEdition/DebugDependencyGraph.h"
+#include "HistoryEdition/DependencyGraphBuilder.h"
 #include "HistoryEdition/HistoryAnalysis.h"
 #include "RenameEditAndDeleteMapsFlow.h"
-#include "ScopedSessionDatabase.h"
+#include "Util/ScopedSessionDatabase.h"
 
 #include "Algo/AllOf.h"
-
 #include "Core/Tests/Containers/TestUtils.h"
-
-#include "HistoryEdition/DebugDependencyGraph.h"
-#include "HistoryEdition/DependencyGraphBuilder.h"
-
 #include "Misc/AutomationTest.h"
 
 namespace UE::ConcertSyncTests::AnalysisTests
diff --git a/Engine/Plugins/Developer/Concert/ConcertSync/ConcertSyncTest/Source/ConcertSyncTest/Private/HistoryEdition/HistoryTestUtil.h b/Engine/Plugins/Developer/Concert/ConcertSync/ConcertSyncTest/Source/ConcertSyncTest/Private/HistoryEdition/HistoryTestUtil.h
--- a/Engine/Plugins/Developer/Concert/ConcertSync/ConcertSyncTest/Source/ConcertSyncTest/Private/HistoryEdition/HistoryTestUtil.h
+++ b/Engine/Plugins/Developer/Concert/ConcertSync/ConcertSyncTest/Source/ConcertSyncTest/Private/HistoryEdition/HistoryTestUtil.h
@@ -3,8 +3,6 @@
 #pragma once
 
 #include "CoreMinimal.h"
-#include "ConcertMessageData.h"
-#include "ScopedSessionDatabase.h"
 #include "HistoryEdition/ActivityGraphIDs.h"
 #include "HistoryEdition/ActivityNode.h"
 #include "Misc/AutomationTest.h"
@@ -43,102 +41,4 @@ namespace UE::ConcertSyncTests
 
 		return ActivityNodes;
 	}
-
-	/**
-	 * Helps fill the database
-	 */
-	class FTestActivityBuilder
-	{
-		TArray<FActivityID> Activities;
-		FScopedSessionDatabase& SessionDatabase;
-	public:
-
-		FTestActivityBuilder(FScopedSessionDatabase& SessionDatabase, uint32 ActivityCount)
-			: SessionDatabase(SessionDatabase)
-		{
-			Activities.SetNumUninitialized(ActivityCount);
-		}
-
-		TArray<FActivityID> GetActivities() const { return Activities; }
-
-		bool NewMap(FName MapName, uint32 ActivityIndex)
-		{
-			FConcertSyncActivity NewPackage;
-			NewPackage.EndpointId = SessionDatabase.GetEndpoint();
-			FConcertPackageInfo PackageInfo;
-			FConcertPackageDataStream PackageDataStream;
-			PackageInfo.PackageName = MapName;
-			PackageInfo.PackageUpdateType = EConcertPackageUpdateType::Added;
-			SessionDatabase.GetTransactionMaxEventId(PackageInfo.TransactionEventIdAtSave);
-			int64 Dummy;
-			return SessionDatabase.AddPackageActivity(NewPackage, PackageInfo, PackageDataStream, Activities[ActivityIndex], Dummy);
-		}
-
-		bool SaveMap(FName MapName, uint32 ActivityIndex)
-		{
-			FConcertSyncActivity SaveFooPackage;
-			SaveFooPackage.EndpointId = SessionDatabase.GetEndpoint();
-			FConcertPackageInfo PackageInfo;
-			FConcertPackageDataStream PackageDataStream;
-			PackageInfo.PackageName = MapName;
-			PackageInfo.PackageUpdateType = EConcertPackageUpdateType::Saved;
-			SessionDatabase.GetTransactionMaxEventId(PackageInfo.TransactionEventIdAtSave);
-			int64 Dummy;
-			return SessionDatabase.AddPackageActivity(SaveFooPackage, PackageInfo, PackageDataStream, Activities[ActivityIndex], Dummy);
-		}
-
-		bool DeleteMap(FName MapName, uint32 ActivityIndex)
-		{
-			FConcertSyncActivity DeleteFooPackage;
-			DeleteFooPackage.EndpointId = SessionDatabase.GetEndpoint();
-			FConcertPackageInfo PackageInfo;
-			FConcertPackageDataStream PackageDataStream;
-			PackageInfo.PackageName = MapName;
-			PackageInfo.PackageUpdateType = EConcertPackageUpdateType::Deleted;
-			SessionDatabase.GetTransactionMaxEventId(PackageInfo.TransactionEventIdAtSave);
-			int64 Dummy;
-			return SessionDatabase.AddPackageActivity(DeleteFooPackage, PackageInfo, PackageDataStream, Activities[ActivityIndex], Dummy);
-		}
-
-		bool RenameMap(FName OldMapName, FName NewMapName, uint32 SaveActivityIndex, uint32 RenameActivityIndex)
-		{
-			int64 Dummy;
-			
-			FConcertSyncActivity SaveFoo2Package;
-			SaveFoo2Package.EndpointId = SessionDatabase.GetEndpoint();
-			FConcertPackageInfo PackageInfo;
-			FConcertPackageDataStream PackageDataStream;
-			PackageInfo.PackageName = NewMapName;
-			PackageInfo.PackageUpdateType = EConcertPackageUpdateType::Saved;
-			SessionDatabase.GetTransactionMaxEventId(PackageInfo.TransactionEventIdAtSave);
-			bool bSuccess = SessionDatabase.AddPackageActivity(SaveFoo2Package, PackageInfo, PackageDataStream, Activities[SaveActivityIndex], Dummy);
-			
-			PackageInfo.PackageUpdateType = EConcertPackageUpdateType::Renamed;
-			PackageInfo.PackageName = OldMapName;
-			PackageInfo.NewPackageName = NewMapName;
-			bSuccess &= SessionDatabase.AddPackageActivity(SaveFoo2Package, PackageInfo, PackageDataStream, Activities[RenameActivityIndex], Dummy);
-
-			return bSuccess;
-		}
-
-		bool CreateActor(FName MapName, uint32 ActivityIndex, FName ActorName = EName::Actor)
-		{
-			FConcertExportedObject Actor;
-			Actor.ObjectId.ObjectName = ActorName;
-			Actor.ObjectId.ObjectPackageName = MapName;
-			Actor.ObjectId.ObjectOuterPathName = *FString::Printf(TEXT("%s:PersistentLevel"), *MapName.ToString());
-			Actor.ObjectId.ObjectClassPathName = TEXT("/Script/Engine.StaticMeshActor");
-			
-			FConcertSyncTransactionActivity CreateActorActivity;
-			CreateActorActivity.EndpointId = SessionDatabase.GetEndpoint();
-			CreateActorActivity.EventData.Transaction.TransactionId = FGuid::NewGuid();
-			CreateActorActivity.EventData.Transaction.OperationId = FGuid::NewGuid();
-			Actor.ObjectData.bAllowCreate = true;
-			CreateActorActivity.EventData.Transaction.ExportedObjects = { Actor };
-			CreateActorActivity.EventData.Transaction.ModifiedPackages = { MapName };
-			SessionDatabase.GetTransactionMaxEventId(CreateActorActivity.EventId);
-			int64 Dummy;
-			return SessionDatabase.AddTransactionActivity(CreateActorActivity, Activities[ActivityIndex], Dummy);
-		}
-	};
 }
diff --git a/Engine/Plugins/Developer/Concert/ConcertSync/ConcertSyncTest/Source/ConcertSyncTest/Private/HistoryEdition/HistoryTopDownAnalysisTests.cpp b/Engine/Plugins/Developer/Concert/ConcertSync/ConcertSyncTest/Source/ConcertSyncTest/Private/HistoryEdition/HistoryTopDownAnalysisTests.cpp
--- a/Engine/Plugins/Developer/Concert/ConcertSync/ConcertSyncTest/Source/ConcertSyncTest/Private/HistoryEdition/HistoryTopDownAnalysisTests.cpp
+++ b/Engine/Plugins/Developer/Concert/ConcertSync/ConcertSyncTest/Source/ConcertSyncTest/Private/HistoryEdition/HistoryTopDownAnalysisTests.cpp
@@ -3,17 +3,14 @@
 #include "ConcertLogGlobal.h"
 
 #include "ConcertSyncSessionDatabase.h"
+#include "HistoryEdition/DebugDependencyGraph.h"
+#include "HistoryEdition/DependencyGraphBuilder.h"
 #include "HistoryEdition/HistoryAnalysis.h"
 #include "RenameEditAndDeleteMapsFlow.h"
-#include "ScopedSessionDatabase.h"
+#include "Util/ScopedSessionDatabase.h"
 
 #include "Algo/AllOf.h"
-
 #include "Core/Tests/Containers/TestUtils.h"
-
-#include "HistoryEdition/DebugDependencyGraph.h"
-#include "HistoryEdition/DependencyGraphBuilder.h"
-
 #include "Misc/AutomationTest.h"
 
 namespace UE::ConcertSyncTests
diff --git a/Engine/Plugins/Developer/Concert/ConcertSync/ConcertSyncTest/Source/ConcertSyncTest/Private/HistoryEdition/RenameEditAndDeleteMapsFlow.cpp b/Engine/Plugins/Developer/Concert/ConcertSync/ConcertSyncTest/Source/ConcertSyncTest/Private/HistoryEdition/RenameEditAndDeleteMapsFlow.cpp
--- a/Engine/Plugins/Developer/Concert/ConcertSync/ConcertSyncTest/Source/ConcertSyncTest/Private/HistoryEdition/RenameEditAndDeleteMapsFlow.cpp
+++ b/Engine/Plugins/Developer/Concert/ConcertSync/ConcertSyncTest/Source/ConcertSyncTest/Private/HistoryEdition/RenameEditAndDeleteMapsFlow.cpp
@@ -4,7 +4,7 @@
 
 #include "ConcertLogGlobal.h"
 #include "ConcertSyncSessionDatabase.h"
-#include "ScopedSessionDatabase.h"
+#include "Util/ScopedSessionDatabase.h"
 
 namespace UE::ConcertSyncTests::RenameEditAndDeleteMapsFlowTest
 {
diff --git a/Engine/Plugins/Developer/Concert/ConcertSync/ConcertSyncTest/Source/ConcertSyncTest/Private/HistoryEdition/ScopedSessionDatabase.h b/Engine/Plugins/Developer/Concert/ConcertSync/ConcertSyncTest/Source/ConcertSyncTest/Private/HistoryEdition/ScopedSessionDatabase.h
deleted file mode 100644
diff --git a/Engine/Plugins/Developer/Concert/ConcertSync/ConcertSyncTest/Source/ConcertSyncTest/Private/Util/ActivityBuilder.h b/Engine/Plugins/Developer/Concert/ConcertSync/ConcertSyncTest/Source/ConcertSyncTest/Private/Util/ActivityBuilder.h
new file mode 100644
--- /dev/null
+++ b/Engine/Plugins/Developer/Concert/ConcertSync/ConcertSyncTest/Source/ConcertSyncTest/Private/Util/ActivityBuilder.h
@@ -0,0 +1,113 @@
+// Copyright Epic Games, Inc. All Rights Reserved.
+
+#pragma once
+
+#include "CoreMinimal.h"
+#include "ConcertMessageData.h"
+#include "ScopedSessionDatabase.h"
+
+namespace UE::ConcertSyncTests
+{
+	/**
+	 * Helps fill the database
+	 */
+	class FActivityBuilder
+	{
+		TArray<FActivityID> Activities;
+		FScopedSessionDatabase& SessionDatabase;
+	public:
+
+		FActivityBuilder(FScopedSessionDatabase& SessionDatabase, uint32 ActivityCount)
+			: SessionDatabase(SessionDatabase)
+		{
+			Activities.SetNumUninitialized(ActivityCount);
+		}
+
+		TArray<FActivityID> GetActivities() const { return Activities; }
+
+		bool NewMap(FName MapName, uint32 ActivityIndex, EConcertSyncActivityFlags Flags = EConcertSyncActivityFlags::None)
+		{
+			FConcertSyncActivity NewPackage;
+			NewPackage.Flags = Flags;
+			NewPackage.EndpointId = SessionDatabase.GetEndpoint();
+			FConcertPackageInfo PackageInfo;
+			FConcertPackageDataStream PackageDataStream;
+			PackageInfo.PackageName = MapName;
+			PackageInfo.PackageUpdateType = EConcertPackageUpdateType::Added;
+			SessionDatabase.GetTransactionMaxEventId(PackageInfo.TransactionEventIdAtSave);
+			int64 Dummy;
+			return SessionDatabase.AddPackageActivity(NewPackage, PackageInfo, PackageDataStream, Activities[ActivityIndex], Dummy);
+		}
+
+		bool SaveMap(FName MapName, uint32 ActivityIndex, EConcertSyncActivityFlags Flags = EConcertSyncActivityFlags::None)
+		{
+			FConcertSyncActivity SaveFooPackage;
+			SaveFooPackage.Flags = Flags;
+			SaveFooPackage.EndpointId = SessionDatabase.GetEndpoint();
+			FConcertPackageInfo PackageInfo;
+			FConcertPackageDataStream PackageDataStream;
+			PackageInfo.PackageName = MapName;
+			PackageInfo.PackageUpdateType = EConcertPackageUpdateType::Saved;
+			SessionDatabase.GetTransactionMaxEventId(PackageInfo.TransactionEventIdAtSave);
+			int64 Dummy;
+			return SessionDatabase.AddPackageActivity(SaveFooPackage, PackageInfo, PackageDataStream, Activities[ActivityIndex], Dummy);
+		}
+
+		bool DeleteMap(FName MapName, uint32 ActivityIndex, EConcertSyncActivityFlags Flags = EConcertSyncActivityFlags::None)
+		{
+			FConcertSyncActivity DeleteFooPackage;
+			DeleteFooPackage.Flags = Flags;
+			DeleteFooPackage.EndpointId = SessionDatabase.GetEndpoint();
+			FConcertPackageInfo PackageInfo;
+			FConcertPackageDataStream PackageDataStream;
+			PackageInfo.PackageName = MapName;
+			PackageInfo.PackageUpdateType = EConcertPackageUpdateType::Deleted;
+			SessionDatabase.GetTransactionMaxEventId(PackageInfo.TransactionEventIdAtSave);
+			int64 Dummy;
+			return SessionDatabase.AddPackageActivity(DeleteFooPackage, PackageInfo, PackageDataStream, Activities[ActivityIndex], Dummy);
+		}
+
+		bool RenameMap(FName OldMapName, FName NewMapName, uint32 SaveActivityIndex, uint32 RenameActivityIndex, EConcertSyncActivityFlags Flags = EConcertSyncActivityFlags::None)
+		{
+			int64 Dummy;
+			
+			FConcertSyncActivity SaveFoo2Package;
+			SaveFoo2Package.Flags = Flags;
+			SaveFoo2Package.EndpointId = SessionDatabase.GetEndpoint();
+			FConcertPackageInfo PackageInfo;
+			FConcertPackageDataStream PackageDataStream;
+			PackageInfo.PackageName = NewMapName;
+			PackageInfo.PackageUpdateType = EConcertPackageUpdateType::Saved;
+			SessionDatabase.GetTransactionMaxEventId(PackageInfo.TransactionEventIdAtSave);
+			bool bSuccess = SessionDatabase.AddPackageActivity(SaveFoo2Package, PackageInfo, PackageDataStream, Activities[SaveActivityIndex], Dummy);
+			
+			PackageInfo.PackageUpdateType = EConcertPackageUpdateType::Renamed;
+			PackageInfo.PackageName = OldMapName;
+			PackageInfo.NewPackageName = NewMapName;
+			bSuccess &= SessionDatabase.AddPackageActivity(SaveFoo2Package, PackageInfo, PackageDataStream, Activities[RenameActivityIndex], Dummy);
+
+			return bSuccess;
+		}
+
+		bool CreateActor(FName MapName, uint32 ActivityIndex, FName ActorName = EName::Actor, EConcertSyncActivityFlags Flags = EConcertSyncActivityFlags::None)
+		{
+			FConcertExportedObject Actor;
+			Actor.ObjectId.ObjectName = ActorName;
+			Actor.ObjectId.ObjectPackageName = MapName;
+			Actor.ObjectId.ObjectOuterPathName = *FString::Printf(TEXT("%s:PersistentLevel"), *MapName.ToString());
+			Actor.ObjectId.ObjectClassPathName = TEXT("/Script/Engine.StaticMeshActor");
+			
+			FConcertSyncTransactionActivity CreateActorActivity;
+			CreateActorActivity.Flags = Flags;
+			CreateActorActivity.EndpointId = SessionDatabase.GetEndpoint();
+			CreateActorActivity.EventData.Transaction.TransactionId = FGuid::NewGuid();
+			CreateActorActivity.EventData.Transaction.OperationId = FGuid::NewGuid();
+			Actor.ObjectData.bAllowCreate = true;
+			CreateActorActivity.EventData.Transaction.ExportedObjects = { Actor };
+			CreateActorActivity.EventData.Transaction.ModifiedPackages = { MapName };
+			SessionDatabase.GetTransactionMaxEventId(CreateActorActivity.EventId);
+			int64 Dummy;
+			return SessionDatabase.AddTransactionActivity(CreateActorActivity, Activities[ActivityIndex], Dummy);
+		}
+	};
+}
diff --git a/Engine/Plugins/Developer/Concert/ConcertSync/ConcertSyncTest/Source/ConcertSyncTest/Private/Util/ScopedSessionDatabase.h b/Engine/Plugins/Developer/Concert/ConcertSync/ConcertSyncTest/Source/ConcertSyncTest/Private/Util/ScopedSessionDatabase.h
new file mode 100644
--- /dev/null
+++ b/Engine/Plugins/Developer/Concert/ConcertSync/ConcertSyncTest/Source/ConcertSyncTest/Private/Util/ScopedSessionDatabase.h
@@ -0,0 +1,46 @@
+﻿// Copyright Epic Games, Inc. All Rights Reserved.
+
+#pragma once
+
+#include "CoreMinimal.h"
+#include "ConcertSyncSessionDatabase.h"
+
+#include "HAL/FileManager.h"
+#include "Misc/AutomationTest.h"
+#include "Misc/Paths.h"
+
+namespace UE::ConcertSyncTests
+{
+	class FScopedSessionDatabase : public FConcertSyncSessionDatabase
+	{
+		FAutomationTestBase& Test;
+		const FString TestSessionPath_Server = FPaths::ProjectIntermediateDir() / TEXT("ConcertDatabaseTest_Server");
+		const FGuid EndpointID;
+	public:
+
+		FScopedSessionDatabase(FAutomationTestBase& Test)
+			: Test(Test)
+			, EndpointID(FGuid::NewGuid())
+		{
+			Open(TestSessionPath_Server);
+
+			FConcertSyncEndpointData EndpointData;
+			EndpointData.ClientInfo.Initialize();
+			if (!SetEndpoint(EndpointID, EndpointData))
+			{
+				Test.AddError(FString::Printf(TEXT("Test may be faulty because endpoint could not be set: %s"), *GetLastError()));
+			}
+		}
+
+		~FScopedSessionDatabase()
+		{
+			if (IsValid() && !Close())
+			{
+				Test.AddError(TEXT("Failed to close server database"));
+			}
+			IFileManager::Get().DeleteDirectory(*TestSessionPath_Server, false, true);
+		}
+
+		const FGuid& GetEndpoint() const { return EndpointID; }
+	};
+}
\ No newline at end of file
