Fix missing copyright

diff --git a/Engine/Plugins/Developer/Concert/ConcertSync/ConcertSyncTest/Source/ConcertSyncTest/Private/Replication/TestReflectionObject.h b/Engine/Plugins/Developer/Concert/ConcertSync/ConcertSyncTest/Source/ConcertSyncTest/Private/Replication/TestReflectionObject.h
--- a/Engine/Plugins/Developer/Concert/ConcertSync/ConcertSyncTest/Source/ConcertSyncTest/Private/Replication/TestReflectionObject.h
+++ b/Engine/Plugins/Developer/Concert/ConcertSync/ConcertSyncTest/Source/ConcertSyncTest/Private/Replication/TestReflectionObject.h
@@ -1,4 +1,4 @@
-// Fill out your copyright notice in the Description page of Project Settings.
+// Copyright Epic Games, Inc. All Rights Reserved.
 
 #pragma once
 
