Fix struct properties not being initialized properly

diff --git a/Engine/Plugins/Developer/Concert/ConcertSync/ConcertSyncCore/Source/ConcertSyncCore/Public/Replication/Data/ReplicationStreamDescription.h b/Engine/Plugins/Developer/Concert/ConcertSync/ConcertSyncCore/Source/ConcertSyncCore/Public/Replication/Data/ReplicationStreamDescription.h
--- a/Engine/Plugins/Developer/Concert/ConcertSync/ConcertSyncCore/Source/ConcertSyncCore/Public/Replication/Data/ReplicationStreamDescription.h
+++ b/Engine/Plugins/Developer/Concert/ConcertSync/ConcertSyncCore/Source/ConcertSyncCore/Public/Replication/Data/ReplicationStreamDescription.h
@@ -16,7 +16,7 @@ struct FSharedReplicationStreamDescription
 	
 	/** Unique id for this stream. */
 	UPROPERTY()
-	FGuid Identifier = FGuid::NewGuid();
+	FGuid Identifier;
 
 	/** Identifies the data that this stream will send. */
 	UPROPERTY()
diff --git a/Engine/Plugins/Developer/Concert/ConcertSync/ConcertSyncTest/Source/ConcertSyncTest/Private/Replication/TestReflectionObject.h b/Engine/Plugins/Developer/Concert/ConcertSync/ConcertSyncTest/Source/ConcertSyncTest/Private/Replication/TestReflectionObject.h
--- a/Engine/Plugins/Developer/Concert/ConcertSync/ConcertSyncTest/Source/ConcertSyncTest/Private/Replication/TestReflectionObject.h
+++ b/Engine/Plugins/Developer/Concert/ConcertSync/ConcertSyncTest/Source/ConcertSyncTest/Private/Replication/TestReflectionObject.h
@@ -40,7 +40,7 @@ struct FTestReplicationStruct
 	UPROPERTY()
 	float Value{};
 	UPROPERTY()
-	FVector Vector;
+	FVector Vector = FVector::ZeroVector;
 	UPROPERTY()
 	FNativeStruct NativeStruct;
 
@@ -90,7 +90,7 @@ public:
 	float Float{};
 	
 	UPROPERTY()
-	FVector Vector;
+	FVector Vector = FVector::ZeroVector;
 
 	UPROPERTY()
 	FTestNestedReplicationStruct TestStruct;
