Replicate objects via Concert. This is an inefficient MVP prototype which just serializes the entire objects and sends it all.

diff --git a/Engine/Plugins/Developer/Concert/ConcertApp/MultiUserClient/Source/MultiUserReplicationEditor/Private/Settings/MultiUserReplicationSettings.cpp b/Engine/Plugins/Developer/Concert/ConcertApp/MultiUserClient/Source/MultiUserReplicationEditor/Private/Settings/MultiUserReplicationSettings.cpp
--- a/Engine/Plugins/Developer/Concert/ConcertApp/MultiUserClient/Source/MultiUserReplicationEditor/Private/Settings/MultiUserReplicationSettings.cpp
+++ b/Engine/Plugins/Developer/Concert/ConcertApp/MultiUserClient/Source/MultiUserReplicationEditor/Private/Settings/MultiUserReplicationSettings.cpp
@@ -32,7 +32,7 @@ namespace UE::MultiUserReplicationEditor::Private
 			const bool bMatches = IndexOfMatches != INDEX_NONE;
 			if (bMatches)
 			{
-				Info.ReplicatedProperties.ReplicatedProperties.Emplace(&Chain, LeafProperty);
+				Info.PropertySelection.ReplicatedProperties.Emplace(&Chain, LeafProperty);
 				Paths.RemoveAtSwap(IndexOfMatches);
 			}
 			return true;
diff --git a/Engine/Plugins/Developer/Concert/ConcertApp/MultiUserClient/Source/MultiUserReplicationEditor/Private/StreamEditor/Model/PropertySelectionAssetModel.cpp b/Engine/Plugins/Developer/Concert/ConcertApp/MultiUserClient/Source/MultiUserReplicationEditor/Private/StreamEditor/Model/PropertySelectionAssetModel.cpp
--- a/Engine/Plugins/Developer/Concert/ConcertApp/MultiUserClient/Source/MultiUserReplicationEditor/Private/StreamEditor/Model/PropertySelectionAssetModel.cpp
+++ b/Engine/Plugins/Developer/Concert/ConcertApp/MultiUserClient/Source/MultiUserReplicationEditor/Private/StreamEditor/Model/PropertySelectionAssetModel.cpp
@@ -34,7 +34,7 @@ namespace UE::MultiUserReplicationEditor
 
 		const FReplicatedObjectInfo* AssignedProperties = Asset->ReplicationMap.ReplicatedObjects.Find(Object);
 		return AssignedProperties
-			? AssignedProperties->ReplicatedProperties.ReplicatedProperties.Num()
+			? AssignedProperties->PropertySelection.ReplicatedProperties.Num()
 			: 0;
 	}
 
@@ -66,7 +66,7 @@ namespace UE::MultiUserReplicationEditor
 
 		const FReplicatedObjectInfo* ObjectInfo = Asset->ReplicationMap.ReplicatedObjects.Find(Object);
 		return ObjectInfo
-			&& Algo::AllOf(Properties, [ObjectInfo](const FConcertPropertyChain& Property){ return ObjectInfo->ReplicatedProperties.ReplicatedProperties.Contains(Property); });
+			&& Algo::AllOf(Properties, [ObjectInfo](const FConcertPropertyChain& Property){ return ObjectInfo->PropertySelection.ReplicatedProperties.Contains(Property); });
 	}
 
 	bool FPropertySelectionAssetModel::ForEachReplicatedObject(TFunctionRef<EBreakBehavior(const FSoftObjectPath& Object)> Delegate) const
@@ -100,14 +100,14 @@ namespace UE::MultiUserReplicationEditor
 			return false;
 		}
 
-		for (const FConcertPropertyChain& ReplicatedPropertyInfo : AssignedProperties->ReplicatedProperties.ReplicatedProperties)
+		for (const FConcertPropertyChain& ReplicatedPropertyInfo : AssignedProperties->PropertySelection.ReplicatedProperties)
 		{
 			if (Delegate(ReplicatedPropertyInfo) == EBreakBehavior::Break)
 			{
 				return true;
 			}
 		}
-		return !AssignedProperties->ReplicatedProperties.ReplicatedProperties.IsEmpty();
+		return !AssignedProperties->PropertySelection.ReplicatedProperties.IsEmpty();
 	}
 
 	void FPropertySelectionAssetModel::AddObjects(TArrayView<UObject*> Objects)
@@ -217,7 +217,7 @@ namespace UE::MultiUserReplicationEditor
 		}
 
 		bool bAddedAtLeastOne = false;
-		TArray<FConcertPropertyChain>& ReplicatedProperties = AssignedProperties->ReplicatedProperties.ReplicatedProperties;
+		TArray<FConcertPropertyChain>& ReplicatedProperties = AssignedProperties->PropertySelection.ReplicatedProperties;
 		ReplicatedProperties.Reserve(ReplicatedProperties.Num() + Properties.Num());
 		for (const FConcertPropertyChain& AddedProperty : Properties)
 		{
@@ -229,7 +229,7 @@ namespace UE::MultiUserReplicationEditor
 			{
 				if (Property.IsParentOf(AddedProperty))
 				{
-					bAddedAtLeastOne |= AssignedProperties->ReplicatedProperties.ReplicatedProperties.AddUnique(Property) != INDEX_NONE;
+					bAddedAtLeastOne |= AssignedProperties->PropertySelection.ReplicatedProperties.AddUnique(Property) != INDEX_NONE;
 				}
 				return EBreakBehavior::Continue;
 			});
@@ -261,7 +261,7 @@ namespace UE::MultiUserReplicationEditor
 		int32 NumRemoved = 0;
 		for (const FConcertPropertyChain& RemovedProperty : Properties)
 		{
-			NumRemoved += AssignedProperties->ReplicatedProperties.ReplicatedProperties.Remove(RemovedProperty);
+			NumRemoved += AssignedProperties->PropertySelection.ReplicatedProperties.Remove(RemovedProperty);
 
 			// Removal should not fail if the class is not available
 			if (!Class)
@@ -275,7 +275,7 @@ namespace UE::MultiUserReplicationEditor
 			{
 				if (Property.IsChildOf(RemovedProperty))
 				{
-					NumRemoved += AssignedProperties->ReplicatedProperties.ReplicatedProperties.Remove(Property);
+					NumRemoved += AssignedProperties->PropertySelection.ReplicatedProperties.Remove(Property);
 				}
 				return EBreakBehavior::Continue;
 			});
diff --git a/Engine/Plugins/Developer/Concert/ConcertSync/ConcertSyncClient/Source/ConcertSyncClient/Private/ConcertBridges.h b/Engine/Plugins/Developer/Concert/ConcertSync/ConcertSyncClient/Source/ConcertSyncClient/Private/ConcertBridges.h
new file mode 100644
--- /dev/null
+++ b/Engine/Plugins/Developer/Concert/ConcertSync/ConcertSyncClient/Source/ConcertSyncClient/Private/ConcertBridges.h
@@ -0,0 +1,23 @@
+// Copyright Epic Games, Inc. All Rights Reserved.
+
+#pragma once
+
+class IConcertClientPackageBridge;
+class IConcertClientReplicationBridge;
+class IConcertClientTransactionBridge;
+
+namespace UE::ConcertSyncClient
+{
+	struct FConcertBridges
+	{
+		/** Package bridge used by sessions of this client */
+		IConcertClientPackageBridge* PackageBridge;
+		/** Transaction bridge used by sessions of this client */
+		IConcertClientTransactionBridge* TransactionBridge;
+		/** Replication bridge used by sessions of this client */
+		IConcertClientReplicationBridge* ReplicationBridge;
+
+		bool IsValid() const { return PackageBridge && TransactionBridge && ReplicationBridge; }
+	};
+}
+
diff --git a/Engine/Plugins/Developer/Concert/ConcertSync/ConcertSyncClient/Source/ConcertSyncClient/Private/ConcertClientWorkspace.cpp b/Engine/Plugins/Developer/Concert/ConcertSync/ConcertSyncClient/Source/ConcertSyncClient/Private/ConcertClientWorkspace.cpp
--- a/Engine/Plugins/Developer/Concert/ConcertSync/ConcertSyncClient/Source/ConcertSyncClient/Private/ConcertClientWorkspace.cpp
+++ b/Engine/Plugins/Developer/Concert/ConcertSync/ConcertSyncClient/Source/ConcertSyncClient/Private/ConcertClientWorkspace.cpp
@@ -1,21 +1,18 @@
 // Copyright Epic Games, Inc. All Rights Reserved.
 
 #include "ConcertClientWorkspace.h"
-#include "Algo/AllOf.h"
+
 #include "ConcertClientTransactionManager.h"
 #include "ConcertClientPackageManager.h"
 #include "ConcertClientLockManager.h"
 #include "IConcertClientPackageBridge.h"
 #include "IConcertClient.h"
 #include "IConcertClientWorkspace.h"
-#include "IConcertModule.h"
+#include "IConcertSyncClient.h"
 #include "IConcertSyncClientModule.h"
 
-#include "IConcertSession.h"
-#include "IConcertFileSharingService.h"
 #include "ConcertSyncClientLiveSession.h"
 #include "ConcertSyncSessionDatabase.h"
-#include "ConcertSyncSettings.h"
 #include "ConcertClientSettings.h"
 #include "ConcertSyncClientUtil.h"
 #include "ConcertLogGlobal.h"
@@ -23,28 +20,20 @@
 #include "ConcertWorkspaceMessages.h"
 #include "ConcertClientDataStore.h"
 #include "ConcertClientLiveTransactionAuthors.h"
+#include "IConcertSession.h"
+#include "IConcertFileSharingService.h"
 
-
-#include "Containers/Ticker.h"
+#include "Algo/AllOf.h"
 #include "Containers/ArrayBuilder.h"
-#include "IConcertSyncClient.h"
 #include "UObject/Package.h"
-#include "UObject/Linker.h"
-#include "UObject/LinkerLoad.h"
 #include "UObject/SavePackage.h"
 #include "UObject/StructOnScope.h"
 #include "Misc/PackageName.h"
 #include "HAL/FileManager.h"
-#include "HAL/PlatformFileManager.h"
 #include "HAL/IConsoleManager.h"
-#include "Misc/App.h"
-#include "Misc/Paths.h"
-#include "Misc/FileHelper.h"
 #include "Misc/CoreDelegates.h"
 #include "Misc/ScopedSlowTask.h"
 #include "Misc/FeedbackContext.h"
-#include "RenderingThread.h"
-#include "Modules/ModuleManager.h"
 #include "StructDeserializer.h"
 #include "StructSerializer.h"
 #include "Backends/JsonStructDeserializerBackend.h"
@@ -55,13 +44,11 @@
 #include "Engine/Engine.h"
 
 #if WITH_EDITOR
+	#include "Editor.h"
 	#include "UnrealEdGlobals.h"
-	#include "UnrealEdMisc.h"
 	#include "Editor/EditorEngine.h"
 	#include "Editor/UnrealEdEngine.h"
 	#include "Editor/TransBuffer.h"
-	#include "FileHelpers.h"
-	#include "GameMapsSettings.h"
 #endif
 
 LLM_DEFINE_TAG(Concert_ConcertClientWorkspace);
@@ -208,18 +195,14 @@ struct FConcertWorkspaceConsoleCommands
 	FAutoConsoleCommand DisableRemoteVerboseLogging;
 };
 
-FConcertClientWorkspace::FConcertClientWorkspace(TSharedRef<FConcertSyncClientLiveSession> InLiveSession,
-												 IConcertClientPackageBridge* InPackageBridge,
-												 IConcertClientTransactionBridge* InTransactionBridge,
-												 TSharedPtr<IConcertFileSharingService> InFileSharingService,
-												 IConcertSyncClient* InOwnerSyncClient)
-	: OwnerSyncClient(InOwnerSyncClient),
-	  FileSharingService(MoveTemp(InFileSharingService))
+FConcertClientWorkspace::FConcertClientWorkspace(const UE::ConcertSyncClient::FSessionBindArgs& SessionBindArgs, TSharedPtr<IConcertFileSharingService> InFileSharingService, IConcertSyncClient* InOwnerSyncClient)
+	: OwnerSyncClient(InOwnerSyncClient)
+	, FileSharingService(MoveTemp(InFileSharingService))
 {
 	static FConcertWorkspaceConsoleCommands ConsoleCommands;
 
 	check(OwnerSyncClient);
-	BindSession(InLiveSession, InPackageBridge, InTransactionBridge);
+	BindSession(SessionBindArgs);
 }
 
 FConcertClientWorkspace::~FConcertClientWorkspace()
@@ -468,15 +451,13 @@ IConcertClientDataStore& FConcertClientWorkspace::GetDataStore()
 	return *DataStore;
 }
 
-void FConcertClientWorkspace::BindSession(TSharedPtr<FConcertSyncClientLiveSession> InLiveSession, IConcertClientPackageBridge* InPackageBridge, IConcertClientTransactionBridge* InTransactionBridge)
+void FConcertClientWorkspace::BindSession(const UE::ConcertSyncClient::FSessionBindArgs& SessionBindArgs)
 {
-	check(InLiveSession->IsValidSession());
-	check(InPackageBridge);
-	check(InTransactionBridge);
+	check(SessionBindArgs.IsValid());
 
 	UnbindSession();
-	LiveSession = InLiveSession;
-	PackageBridge = InPackageBridge;
+	LiveSession = SessionBindArgs.LiveSession;
+	PackageBridge = SessionBindArgs.Bridges.PackageBridge;
 
 	LoadSessionData();
 
@@ -489,13 +470,13 @@ void FConcertClientWorkspace::BindSession(TSharedPtr<FConcertSyncClientLiveSessi
 	// Create Transaction Manager
 	if (EnumHasAnyFlags(LiveSession->GetSessionFlags(), EConcertSyncSessionFlags::EnableTransactions))
 	{
-		TransactionManager = MakeUnique<FConcertClientTransactionManager>(LiveSession.ToSharedRef(), InTransactionBridge);
+		TransactionManager = MakeUnique<FConcertClientTransactionManager>(LiveSession.ToSharedRef(), SessionBindArgs.Bridges.TransactionBridge);
 	}
 
 	// Create Package Manager
 	if (EnumHasAnyFlags(LiveSession->GetSessionFlags(), EConcertSyncSessionFlags::EnablePackages))
 	{
-		PackageManager = MakeUnique<FConcertClientPackageManager>(LiveSession.ToSharedRef(), InPackageBridge, FileSharingService);
+		PackageManager = MakeUnique<FConcertClientPackageManager>(LiveSession.ToSharedRef(), SessionBindArgs.Bridges.PackageBridge, FileSharingService);
 	}
 
 	// Create Lock Manager
diff --git a/Engine/Plugins/Developer/Concert/ConcertSync/ConcertSyncClient/Source/ConcertSyncClient/Private/ConcertClientWorkspace.h b/Engine/Plugins/Developer/Concert/ConcertSync/ConcertSyncClient/Source/ConcertSyncClient/Private/ConcertClientWorkspace.h
--- a/Engine/Plugins/Developer/Concert/ConcertSync/ConcertSyncClient/Source/ConcertSyncClient/Private/ConcertClientWorkspace.h
+++ b/Engine/Plugins/Developer/Concert/ConcertSync/ConcertSyncClient/Source/ConcertSyncClient/Private/ConcertClientWorkspace.h
@@ -3,14 +3,15 @@
 #pragma once
 
 #include "CoreMinimal.h"
-#include "IConcertSyncClient.h"
-#include "Misc/ITransaction.h"
-#include "IConcertClientWorkspace.h"
-#include "IConcertSessionHandler.h"
+#include "ConcertBridges.h"
 #include "ConcertSyncSessionFlags.h"
 #include "ConcertWorkspaceMessages.h"
 #include "ConcertClientWorkspaceData.h"
+#include "ConcertSyncClientLiveSession.h"
 #include "Delegates/Delegate.h"
+#include "IConcertSyncClient.h"
+#include "IConcertClientWorkspace.h"
+#include "IConcertSessionHandler.h"
 
 class IConcertClientSession;
 class IConcertClientPackageBridge;
@@ -24,6 +25,17 @@ class FConcertSyncClientLiveSession;
 class ISourceControlProvider;
 class FConcertClientDataStore;
 
+namespace UE::ConcertSyncClient
+{
+	struct FSessionBindArgs
+	{
+		TSharedPtr<FConcertSyncClientLiveSession> LiveSession;
+		FConcertBridges Bridges;
+		
+		bool IsValid() const { return LiveSession->IsValidSession() && Bridges.IsValid(); }
+	};
+}
+
 struct FScopedSlowTask;
 
 DECLARE_MULTICAST_DELEGATE(FOnWorkspaceEndFrameCompleted);
@@ -31,7 +43,7 @@ DECLARE_MULTICAST_DELEGATE(FOnWorkspaceEndFrameCompleted);
 class FConcertClientWorkspace : public IConcertClientWorkspace
 {
 public:
-	FConcertClientWorkspace(TSharedRef<FConcertSyncClientLiveSession> InLiveSession, IConcertClientPackageBridge* InPackageBridge, IConcertClientTransactionBridge* InTransactionBridge, TSharedPtr<IConcertFileSharingService> InFileSharingService, IConcertSyncClient* InOwnerSyncClient);
+	FConcertClientWorkspace(const UE::ConcertSyncClient::FSessionBindArgs& SessionBindArgs, TSharedPtr<IConcertFileSharingService> InFileSharingService, IConcertSyncClient* InOwnerSyncClient);
 	virtual ~FConcertClientWorkspace();
 
 	// IConcertClientWorkspace interface
@@ -80,7 +92,7 @@ public:
 private:
 
 	/** Bind the workspace to this session. */
-	void BindSession(TSharedPtr<FConcertSyncClientLiveSession> InLiveSession, IConcertClientPackageBridge* InPackageBridge, IConcertClientTransactionBridge* InTransactionBridge);
+	void BindSession(const UE::ConcertSyncClient::FSessionBindArgs& SessionBindArgs);
 
 	/** Unbind the workspace to its bound session. */
 	void UnbindSession();
diff --git a/Engine/Plugins/Developer/Concert/ConcertSync/ConcertSyncClient/Source/ConcertSyncClient/Private/ConcertSyncClient.cpp b/Engine/Plugins/Developer/Concert/ConcertSync/ConcertSyncClient/Source/ConcertSyncClient/Private/ConcertSyncClient.cpp
--- a/Engine/Plugins/Developer/Concert/ConcertSync/ConcertSyncClient/Source/ConcertSyncClient/Private/ConcertSyncClient.cpp
+++ b/Engine/Plugins/Developer/Concert/ConcertSync/ConcertSyncClient/Source/ConcertSyncClient/Private/ConcertSyncClient.cpp
@@ -11,22 +11,24 @@
 #include "ConcertClientSequencerManager.h"
 #include "ConcertClientPresenceManager.h"
 #include "ConcertSourceControlProxy.h"
+#include "Replication/Manager/ReplicationManager.h"
 
 LLM_DECLARE_TAG(Concert_ConcertSyncClient);
 #define LOCTEXT_NAMESPACE "ConcertSyncClient"
 
-FConcertSyncClient::FConcertSyncClient(const FString& InRole, IConcertClientPackageBridge* InPackageBridge, IConcertClientTransactionBridge* InTransactionBridge)
+FConcertSyncClient::FConcertSyncClient(
+	const FString& InRole,
+	const UE::ConcertSyncClient::FConcertBridges& InBridges
+	)
 	: ConcertClient(IConcertClientModule::Get().CreateClient(InRole))
 	, SessionFlags(EConcertSyncSessionFlags::None)
-	, PackageBridge(InPackageBridge)
-	, TransactionBridge(InTransactionBridge)
+	, Bridges(InBridges)
 #if WITH_EDITOR
 	, SequencerManager(MakeUnique<FConcertClientSequencerManager>(this))
 	, SourceControlProxy(MakeUnique<FConcertSourceControlProxy>())
 #endif
 {
-	check(PackageBridge);
-	check(TransactionBridge);
+	check(Bridges.IsValid());
 
 	ConcertClient->OnSessionStartup().AddRaw(this, &FConcertSyncClient::RegisterConcertSyncHandlers);
 	ConcertClient->OnSessionShutdown().AddRaw(this, &FConcertSyncClient::UnregisterConcertSyncHandlers);
@@ -87,6 +89,11 @@ IConcertClientSequencerManager* FConcertSyncClient::GetSequencerManager() const
 	return Manager;
 }
 
+IConcertClientReplicationManager* FConcertSyncClient::GetReplicationManager() const
+{
+	return ReplicationManager.Get();
+}
+
 FOnConcertClientWorkspaceStartupOrShutdown& FConcertSyncClient::OnWorkspaceStartup()
 {
 	return OnWorkspaceStartupDelegate;
@@ -147,7 +154,7 @@ void FConcertSyncClient::SetFileSharingService(TSharedPtr<IConcertFileSharingSer
 void FConcertSyncClient::CreateWorkspace(const TSharedRef<FConcertSyncClientLiveSession>& InLiveSession)
 {
 	DestroyWorkspace();
-	Workspace = MakeShared<FConcertClientWorkspace>(InLiveSession, PackageBridge, TransactionBridge, FileSharingService, this);
+	Workspace = MakeShared<FConcertClientWorkspace>(UE::ConcertSyncClient::FSessionBindArgs{ InLiveSession, Bridges } , FileSharingService, this);
 	OnWorkspaceStartupDelegate.Broadcast(Workspace);
 #if WITH_EDITOR
 	if (GIsEditor && EnumHasAllFlags(SessionFlags, EConcertSyncSessionFlags::EnablePackages | EConcertSyncSessionFlags::ShouldUsePackageSandbox))
@@ -178,6 +185,14 @@ void FConcertSyncClient::RegisterConcertSyncHandlers(TSharedRef<IConcertClientSe
 	{
 		OnSyncSessionStartupDelegate.Broadcast(this);
 		CreateWorkspace(LiveSession.ToSharedRef());
+		
+		// Create Replication Manager
+		if (EnumHasAnyFlags(LiveSession->GetSessionFlags(), EConcertSyncSessionFlags::EnableReplication))
+		{
+			ReplicationManager = MakeUnique<UE::ConcertSyncClient::Replication::FReplicationManager>(InSession, Bridges.ReplicationBridge);
+			ReplicationManager->StartAcceptingJoinRequests();
+		}
+
 #if WITH_EDITOR
 		PresenceManager.Reset();
 		if (EnumHasAnyFlags(SessionFlags, EConcertSyncSessionFlags::EnablePresence))
@@ -198,12 +213,17 @@ void FConcertSyncClient::RegisterConcertSyncHandlers(TSharedRef<IConcertClientSe
 
 IConcertClientTransactionBridge* FConcertSyncClient::GetTransactionBridge() const
 {
-	return TransactionBridge;
+	return Bridges.TransactionBridge;
 }
 
 IConcertClientPackageBridge* FConcertSyncClient::GetPackageBridge() const
 {
-	return PackageBridge;
+	return Bridges.PackageBridge;
+}
+
+IConcertClientReplicationBridge* FConcertSyncClient::GetReplicationBridge() const
+{
+	return Bridges.ReplicationBridge;
 }
 
 void FConcertSyncClient::UnregisterConcertSyncHandlers(TSharedRef<IConcertClientSession> InSession)
@@ -212,6 +232,7 @@ void FConcertSyncClient::UnregisterConcertSyncHandlers(TSharedRef<IConcertClient
 	SequencerManager->Unregister(InSession);
 	PresenceManager.Reset();
 #endif
+	ReplicationManager.Reset();
 	DestroyWorkspace();
 	OnSyncSessionShutdownDelegate.Broadcast(this);
 	LiveSession.Reset();
diff --git a/Engine/Plugins/Developer/Concert/ConcertSync/ConcertSyncClient/Source/ConcertSyncClient/Private/ConcertSyncClient.h b/Engine/Plugins/Developer/Concert/ConcertSync/ConcertSyncClient/Source/ConcertSyncClient/Private/ConcertSyncClient.h
--- a/Engine/Plugins/Developer/Concert/ConcertSync/ConcertSyncClient/Source/ConcertSyncClient/Private/ConcertSyncClient.h
+++ b/Engine/Plugins/Developer/Concert/ConcertSync/ConcertSyncClient/Source/ConcertSyncClient/Private/ConcertSyncClient.h
@@ -2,10 +2,17 @@
 
 #pragma once
 
+#include "ConcertBridges.h"
 #include "IConcertSyncClient.h"
 
+namespace UE::ConcertSyncClient::Replication
+{
+	class FReplicationManager;
+}
+
 class IConcertClientSession;
 class IConcertClientPackageBridge;
+class IConcertClientReplicationBridge;
 class IConcertClientTransactionBridge;
 class FConcertSyncClientLiveSession;
 class FConcertClientWorkspace;
@@ -19,7 +26,7 @@ class FConcertSourceControlProxy;
 class FConcertSyncClient : public IConcertSyncClient
 {
 public:
-	FConcertSyncClient(const FString& InRole, IConcertClientPackageBridge* InPackageBridge, IConcertClientTransactionBridge* InTransactionBridge);
+	FConcertSyncClient(const FString& InRole, const UE::ConcertSyncClient::FConcertBridges& InBridges);
 	virtual ~FConcertSyncClient();
 
 	//~ IConcertSyncClient interface
@@ -29,6 +36,7 @@ public:
 	virtual TSharedPtr<IConcertClientWorkspace> GetWorkspace() const override;
 	virtual IConcertClientPresenceManager* GetPresenceManager() const override;
 	virtual IConcertClientSequencerManager* GetSequencerManager() const override;
+	virtual IConcertClientReplicationManager* GetReplicationManager() const override;
 
 	virtual FOnConcertClientWorkspaceStartupOrShutdown& OnWorkspaceStartup() override;
 	virtual FOnConcertClientWorkspaceStartupOrShutdown& OnWorkspaceShutdown() override;
@@ -43,6 +51,7 @@ public:
 	virtual void SetFileSharingService(TSharedPtr<IConcertFileSharingService> InFileSharingService) override;
 	virtual IConcertClientTransactionBridge* GetTransactionBridge() const override;
 	virtual IConcertClientPackageBridge* GetPackageBridge() const override;
+	virtual IConcertClientReplicationBridge* GetReplicationBridge() const override;
 
 private:
 	void CreateWorkspace(const TSharedRef<FConcertSyncClientLiveSession>& InLiveSession);
@@ -57,11 +66,8 @@ private:
 	/** Flags controlling what features are enabled for sessions within this client */
 	EConcertSyncSessionFlags SessionFlags;
 
-	/** Package bridge used by sessions of this client */
-	IConcertClientPackageBridge* PackageBridge;
-
-	/** Transaction bridge used by sessions of this client */
-	IConcertClientTransactionBridge* TransactionBridge;
+	/** Pointers to Concert Bridges used by this client. */
+	const UE::ConcertSyncClient::FConcertBridges Bridges;
 
 	/** Live session data for the current session. */
 	TSharedPtr<FConcertSyncClientLiveSession> LiveSession;
@@ -83,6 +89,9 @@ private:
 
 	/** Optional side channel to exchange large blobs (package data) with the server in a scalable way (ex. the request/response transport layer is not designed and doesn't support exchanging 3GB packages). */
 	TSharedPtr<IConcertFileSharingService> FileSharingService;
+	
+	/** Manages the client's replication system. Can be used to join & leave replication sessions. */
+	TUniquePtr<UE::ConcertSyncClient::Replication::FReplicationManager> ReplicationManager;
 
 #if WITH_EDITOR
 	/** Presence manager for the current session. */
diff --git a/Engine/Plugins/Developer/Concert/ConcertSync/ConcertSyncClient/Source/ConcertSyncClient/Private/ConcertSyncClientModule.cpp b/Engine/Plugins/Developer/Concert/ConcertSync/ConcertSyncClient/Source/ConcertSyncClient/Private/ConcertSyncClientModule.cpp
--- a/Engine/Plugins/Developer/Concert/ConcertSync/ConcertSyncClient/Source/ConcertSyncClient/Private/ConcertSyncClientModule.cpp
+++ b/Engine/Plugins/Developer/Concert/ConcertSync/ConcertSyncClient/Source/ConcertSyncClient/Private/ConcertSyncClientModule.cpp
@@ -8,6 +8,7 @@
 #include "ConcertLogGlobal.h"
 #include "ConcertClientPackageBridge.h"
 #include "ConcertClientTransactionBridge.h"
+#include "Replication/ConcertClientReplicationBridge.h"
 
 #include "Framework/Notifications/NotificationManager.h"
 
@@ -27,6 +28,7 @@ public:
 		LLM_SCOPE_BYTAG(Concert_ConcertSyncClient);
 		PackageBridge = MakeUnique<FConcertClientPackageBridge>();
 		TransactionBridge = MakeUnique<FConcertClientTransactionBridge>();
+		ReplicationBridge = MakeUnique<UE::ConcertSyncClient::Replication::FConcertClientReplicationBridge>();
 	}
 
 	virtual void ShutdownModule() override
@@ -34,6 +36,7 @@ public:
 		LLM_SCOPE_BYTAG(Concert_ConcertSyncClient);
 		PackageBridge.Reset();
 		TransactionBridge.Reset();
+		ReplicationBridge.Reset();
 	}
 
 	virtual UConcertClientConfig* ParseClientSettings(const TCHAR* CommandLine) override
@@ -116,7 +119,12 @@ public:
 		// Remove dead clients.
 		Clients.RemoveAll([](TWeakPtr<IConcertSyncClient> WeakClient) { return !WeakClient.IsValid(); });
 
-		TSharedRef<IConcertSyncClient> NewClient = MakeShared<FConcertSyncClient>(InRole, PackageBridge.Get(), TransactionBridge.Get());
+		const TSharedRef<IConcertSyncClient> NewClient = MakeShared<FConcertSyncClient>(
+			InRole,
+			UE::ConcertSyncClient::FConcertBridges{ PackageBridge.Get(),
+			TransactionBridge.Get(),
+			ReplicationBridge.Get() }
+			);
 		Clients.Add(NewClient);
 		OnClientCreated().Broadcast(NewClient);
 
@@ -133,6 +141,11 @@ public:
 		return *TransactionBridge;
 	}
 
+	virtual IConcertClientReplicationBridge& GetReplicationBridge() override
+	{
+		return *ReplicationBridge;
+	}
+
 	virtual TArray<TSharedRef<IConcertSyncClient>> GetClients() const override
 	{
 		TArray<TSharedRef<IConcertSyncClient>> ActiveClients;
@@ -173,6 +186,7 @@ public:
 private:
 	TUniquePtr<FConcertClientPackageBridge> PackageBridge;
 	TUniquePtr<FConcertClientTransactionBridge> TransactionBridge;
+	TUniquePtr<UE::ConcertSyncClient::Replication::FConcertClientReplicationBridge> ReplicationBridge;
 	TArray<TWeakPtr<IConcertSyncClient>> Clients;
 	FOnConcertClientCreated OnClientCreatedDelegate;
 };
diff --git a/Engine/Plugins/Developer/Concert/ConcertSync/ConcertSyncClient/Source/ConcertSyncClient/Private/Replication/ConcertClientReplicationBridge.cpp b/Engine/Plugins/Developer/Concert/ConcertSync/ConcertSyncClient/Source/ConcertSyncClient/Private/Replication/ConcertClientReplicationBridge.cpp
new file mode 100644
--- /dev/null
+++ b/Engine/Plugins/Developer/Concert/ConcertSync/ConcertSyncClient/Source/ConcertSyncClient/Private/Replication/ConcertClientReplicationBridge.cpp
@@ -0,0 +1,126 @@
+// Copyright Epic Games, Inc. All Rights Reserved.
+
+#include "ConcertClientReplicationBridge.h"
+
+#include "ConcertLogGlobal.h"
+
+#include "AssetRegistry/IAssetRegistry.h"
+#include "AssetRegistry/AssetData.h"
+#include "Engine/World.h"
+
+class IAssetRegistry;
+
+namespace UE::ConcertSyncClient::Replication
+{
+	FConcertClientReplicationBridge::FConcertClientReplicationBridge()
+	{
+		FWorldDelegates::OnPostWorldInitialization.AddRaw(this, &FConcertClientReplicationBridge::OnPostWorldInitialization);
+		FWorldDelegates::OnPostWorldCleanup.AddRaw(this, &FConcertClientReplicationBridge::OnPostWorldCleanup);
+	}
+
+	FConcertClientReplicationBridge::~FConcertClientReplicationBridge()
+	{
+		FWorldDelegates::OnPostWorldInitialization.RemoveAll(this);
+		FWorldDelegates::OnPostWorldCleanup.RemoveAll(this);
+	}
+
+	void FConcertClientReplicationBridge::PushTrackedObjects(TArrayView<const FSoftObjectPath> InTrackedObjects)
+	{
+		TSet<UObject*> ResolvedObjects;
+		
+		for (const FSoftObjectPath& TrackedObjectPath : InTrackedObjects)
+		{
+			FTrackedObject& TrackedObject = TrackedObjects.FindOrAdd(TrackedObjectPath);
+			++TrackedObject.TrackCounter;
+			if (!TrackedObject.ResolvedObject.IsValid())
+			{
+				UObject* ResolvedObject = FindObjectIfAvailable(TrackedObjectPath);
+				TrackedObject.ResolvedObject = ResolvedObject;
+				if (ResolvedObject)
+				{
+					ResolvedObjects.Add(ResolvedObject);
+				}
+			}
+		}
+		
+		// Broadcast after updating TrackedObjects in the unlikely case that the broadcast calls PushTrackedObjects or PopTrackedObjects.
+		for (UObject* ResolvedObject : ResolvedObjects)
+		{
+			OnObjectDiscoveredDelegate.Broadcast(*ResolvedObject);
+		}
+	}
+
+	void FConcertClientReplicationBridge::PopTrackedObjects(TArrayView<const FSoftObjectPath> InTrackedObjects)
+	{
+		for (const FSoftObjectPath& TrackedObjectPath : InTrackedObjects)
+		{
+			FTrackedObject* TrackedObject = TrackedObjects.Find(TrackedObjectPath);
+			if (!TrackedObject)
+			{
+				UE_LOG(LogConcert, Warning, TEXT("Object %s was requested to not be tracked but was it already wasn't."), *TrackedObjectPath.ToString());
+				continue;
+			}
+
+			if (TrackedObject->TrackCounter == 1)
+			{
+				TrackedObjects.Remove(TrackedObjectPath);
+			}
+			else
+			{
+				--TrackedObject->TrackCounter;
+			}
+		}
+	}
+
+	bool FConcertClientReplicationBridge::IsObjectAvailable(const FSoftObjectPath& Path)
+	{
+		// TODO: Possibly iterate through worlds and see whether they're the parent of this object. Profile performance.
+		return FindObjectIfAvailable(Path) != nullptr;
+	}
+
+	UObject* FConcertClientReplicationBridge::FindObjectIfAvailable(const FSoftObjectPath& Path)
+	{
+		IAssetRegistry* AssetRegistry = IAssetRegistry::Get();
+		if (!AssetRegistry)
+		{
+			return nullptr;
+		}
+
+		constexpr bool bSkipInMemory = true;
+		const FAssetData AssetData = AssetRegistry->GetAssetByObjectPath(Path, bSkipInMemory);
+		const bool bIsAsset = AssetData.IsValid() || Path.IsAsset();
+		if (!ensure(!bIsAsset))
+		{
+			UE_LOG(LogConcert, Error, TEXT("Assets are not supported for replication in this version (%s)"), *Path.ToString());
+			return nullptr;
+		}
+		
+		// This should be an object in a UWorld. This will resolve if the world is opened.
+		return Path.ResolveObject();
+	}
+
+	void FConcertClientReplicationBridge::OnPostWorldInitialization(UWorld* World, FWorldInitializationValues WorldInitializationValues)
+	{
+		// We'll allow PIE even though it is currently not explicitly supported; the primary use case is editor or game worlds (if someone tries it in packaged).
+		if (!ensure(World) || World->IsPreviewWorld())
+		{
+			return;
+		}
+		
+		LoadedWorlds.Add(World);
+
+		// TODO: Discover objects
+	}
+
+	void FConcertClientReplicationBridge::OnPostWorldCleanup(UWorld* World, bool bSessionEnded, bool bCleanupResources)
+	{
+		const int32 NumRemoved = LoadedWorlds.Remove(World);
+		bool bRemovedAnything = NumRemoved > 0;
+		if (!bRemovedAnything)
+		{
+			return;
+		}
+
+		// TODO: Hide objects
+	}
+}
diff --git a/Engine/Plugins/Developer/Concert/ConcertSync/ConcertSyncClient/Source/ConcertSyncClient/Private/Replication/ConcertClientReplicationBridge.h b/Engine/Plugins/Developer/Concert/ConcertSync/ConcertSyncClient/Source/ConcertSyncClient/Private/Replication/ConcertClientReplicationBridge.h
new file mode 100644
--- /dev/null
+++ b/Engine/Plugins/Developer/Concert/ConcertSync/ConcertSyncClient/Source/ConcertSyncClient/Private/Replication/ConcertClientReplicationBridge.h
@@ -0,0 +1,58 @@
+// Copyright Epic Games, Inc. All Rights Reserved.
+
+#pragma once
+
+#include "Containers/Map.h"
+#include "Containers/Set.h"
+#include "UObject/SoftObjectPath.h"
+#include "UObject/WeakObjectPtr.h"
+#include "Templates/SharedPointer.h"
+#include "Replication/IConcertClientReplicationBridge.h"
+
+class UObject;
+class UWorld;
+struct FWorldInitializationValues;
+
+namespace UE::ConcertSyncClient::Replication
+{
+	class FConcertClientReplicationBridge : public IConcertClientReplicationBridge
+	{
+	public:
+		
+		FConcertClientReplicationBridge();
+		virtual ~FConcertClientReplicationBridge() override;
+
+		//~ Begin IConcertClientReplicationBridge Interface
+		virtual void PushTrackedObjects(TArrayView<const FSoftObjectPath> InTrackedObjects) override;
+		virtual void PopTrackedObjects(TArrayView<const FSoftObjectPath> InTrackedObjects) override;
+		virtual bool IsObjectAvailable(const FSoftObjectPath& Path) override;
+		virtual UObject* FindObjectIfAvailable(const FSoftObjectPath& Path) override;
+		virtual FConcertClientReplicationBridgeObjectEvent& OnObjectDiscovered() override { return OnObjectDiscoveredDelegate; }
+		virtual FConcertClientReplicationBridgeObjectEvent& OnObjectHidden() override { return OnObjectRemovedDelegate; }
+		//~ End IConcertClientReplicationBridge Interface
+
+	private:
+
+		struct FTrackedObject
+		{
+			/** Tracks the number of PushTrackedObjects calls. This entry is removed upon reaching 0. */
+			int32 TrackCounter = 0;
+			/** Cached object pointer to avoid constant resolving of the soft object ptr. */
+			TWeakObjectPtr<UObject> ResolvedObject;
+		};
+
+		/** The objects that wish to be tracked. */
+		TMap<FSoftObjectPath, FTrackedObject> TrackedObjects;
+
+		/** The set of worlds currently open on this client. */
+		TSet<TWeakObjectPtr<UWorld>> LoadedWorlds;
+	
+		FConcertClientReplicationBridgeObjectEvent OnObjectDiscoveredDelegate;
+		FConcertClientReplicationBridgeObjectEvent OnObjectRemovedDelegate;
+
+		// Callbacks into the engine
+		void OnPostWorldInitialization(UWorld* World, FWorldInitializationValues WorldInitializationValues);
+		void OnPostWorldCleanup(UWorld* World, bool bSessionEnded, bool bCleanupResources);
+	};
+}
+
diff --git a/Engine/Plugins/Developer/Concert/ConcertSync/ConcertSyncClient/Source/ConcertSyncClient/Private/Replication/Manager/ReplicationManager.cpp b/Engine/Plugins/Developer/Concert/ConcertSync/ConcertSyncClient/Source/ConcertSyncClient/Private/Replication/Manager/ReplicationManager.cpp
new file mode 100644
--- /dev/null
+++ b/Engine/Plugins/Developer/Concert/ConcertSync/ConcertSyncClient/Source/ConcertSyncClient/Private/Replication/Manager/ReplicationManager.cpp
@@ -0,0 +1,60 @@
+// Copyright Epic Games, Inc. All Rights Reserved.
+
+#include "ReplicationManager.h"
+
+#include "ConcertLogGlobal.h"
+#include "ConcertSyncClientLiveSession.h"
+#include "IConcertSession.h"
+#include "ReplicationManagerState_Disconnected.h"
+
+#include "Replication/Messages/ConcertReplicationHandshakeMessages.h"
+
+namespace UE::ConcertSyncClient::Replication
+{
+	FReplicationManager::FReplicationManager(TSharedRef<IConcertClientSession> InLiveSession, IConcertClientReplicationBridge* InBridge)
+		: Session(MoveTemp(InLiveSession))
+		, Bridge(InBridge)
+	{}
+
+	FReplicationManager::~FReplicationManager()
+	{}
+
+	void FReplicationManager::StartAcceptingJoinRequests()
+	{
+		checkSlow(!CurrentState.IsValid());
+		CurrentState = MakeShared<FReplicationManagerState_Disconnected>(Session, Bridge, *this);
+	}
+
+	TFuture<FJoinReplicatedSessionResult> FReplicationManager::JoinReplicationSession(FJoinReplicatedSessionArgs Args)
+	{
+		return ensureMsgf(CurrentState, TEXT("StartAcceptingJoinRequests should have been called at this point."))
+			? CurrentState->JoinReplicationSession(MoveTemp(Args))
+			: MakeFulfilledPromise<FJoinReplicatedSessionResult>(EJoinReplicationErrorCode::Cancelled).GetFuture();
+	}
+
+	void FReplicationManager::LeaveReplicationSession()
+	{
+		if (ensureMsgf(CurrentState, TEXT("StartAcceptingJoinRequests should have been called at this point.")))
+		{
+			CurrentState->LeaveReplicationSession();
+		}
+	}
+
+	bool FReplicationManager::CanJoin()
+	{
+		return ensureMsgf(CurrentState, TEXT("StartAcceptingJoinRequests should have been called at this point."))
+			&& CurrentState->CanJoin(); 
+	}
+
+	bool FReplicationManager::IsConnectedToReplicationSession()
+	{
+		return ensureMsgf(CurrentState, TEXT("StartAcceptingJoinRequests should have been called at this point."))
+			&& CurrentState->IsConnectedToReplicationSession(); 
+	}
+
+	void FReplicationManager::OnChangeState(TSharedRef<FReplicationManagerState> NewState)
+	{
+		CurrentState = MoveTemp(NewState);
+	}
+}
+
diff --git a/Engine/Plugins/Developer/Concert/ConcertSync/ConcertSyncClient/Source/ConcertSyncClient/Private/Replication/Manager/ReplicationManager.h b/Engine/Plugins/Developer/Concert/ConcertSync/ConcertSyncClient/Source/ConcertSyncClient/Private/Replication/Manager/ReplicationManager.h
new file mode 100644
--- /dev/null
+++ b/Engine/Plugins/Developer/Concert/ConcertSync/ConcertSyncClient/Source/ConcertSyncClient/Private/Replication/Manager/ReplicationManager.h
@@ -0,0 +1,47 @@
+// Copyright Epic Games, Inc. All Rights Reserved.
+
+#pragma once
+
+#include "Replication/IConcertClientReplicationManager.h"
+#include "Replication/Messages/ConcertReplicationHandshakeMessages.h"
+#include "Templates/SharedPointer.h"
+
+class IConcertClientReplicationBridge;
+class IConcertClientSession;
+
+namespace UE::ConcertSyncClient::Replication
+{
+	class FReplicationManagerState;
+	
+	class FReplicationManager : public IConcertClientReplicationManager
+	{
+		friend FReplicationManagerState; 
+	public:
+		
+		FReplicationManager(TSharedRef<IConcertClientSession> InLiveSession, IConcertClientReplicationBridge* InBridge);
+		virtual ~FReplicationManager() override;
+
+		/** Starts accepting join requests. Must be called separately from constructor because of TSharedFromThis asserting if SharedThis is called in constructor. */
+		void StartAcceptingJoinRequests();
+
+		//~ Begin IConcertClientReplicationManager Interface
+		virtual TFuture<FJoinReplicatedSessionResult> JoinReplicationSession(FJoinReplicatedSessionArgs Args) override;
+		virtual void LeaveReplicationSession() override;
+		virtual bool CanJoin() override;
+		virtual bool IsConnectedToReplicationSession() override;
+		//~ End IConcertClientReplicationManager Interface
+
+	private:
+		
+		/** Session instance this manager was created for. */
+		TSharedRef<IConcertClientSession> Session;
+		/** The replication bridge is responsible for applying received data and generating data to send. */
+		IConcertClientReplicationBridge* Bridge;
+
+		/** The current state this manager is in, e.g. waiting for connection request, connecting, connected, etc. */
+		TSharedPtr<FReplicationManagerState> CurrentState;
+		
+		/** Called by FReplicationManagerState to change the state. */
+		void OnChangeState(TSharedRef<FReplicationManagerState> NewState);
+	};
+}
diff --git a/Engine/Plugins/Developer/Concert/ConcertSync/ConcertSyncClient/Source/ConcertSyncClient/Private/Replication/Manager/ReplicationManagerState.cpp b/Engine/Plugins/Developer/Concert/ConcertSync/ConcertSyncClient/Source/ConcertSyncClient/Private/Replication/Manager/ReplicationManagerState.cpp
new file mode 100644
--- /dev/null
+++ b/Engine/Plugins/Developer/Concert/ConcertSync/ConcertSyncClient/Source/ConcertSyncClient/Private/Replication/Manager/ReplicationManagerState.cpp
@@ -0,0 +1,20 @@
+// Copyright Epic Games, Inc. All Rights Reserved.
+
+#include "ReplicationManagerState.h"
+
+#include "ReplicationManager.h"
+
+namespace UE::ConcertSyncClient::Replication
+{
+	FReplicationManagerState::FReplicationManagerState(FReplicationManager& Owner)
+		: Owner(Owner)
+	{}
+
+	void FReplicationManagerState::ChangeState(TSharedRef<FReplicationManagerState> NewState)
+	{
+		// OnChangeState destroys us
+		TSharedPtr<FReplicationManagerState> Guard = SharedThis(this);
+		Owner.OnChangeState(MoveTemp(NewState));
+		Owner.CurrentState->OnEnterState();
+	}
+}
diff --git a/Engine/Plugins/Developer/Concert/ConcertSync/ConcertSyncClient/Source/ConcertSyncClient/Private/Replication/Manager/ReplicationManagerState.h b/Engine/Plugins/Developer/Concert/ConcertSync/ConcertSyncClient/Source/ConcertSyncClient/Private/Replication/Manager/ReplicationManagerState.h
new file mode 100644
--- /dev/null
+++ b/Engine/Plugins/Developer/Concert/ConcertSync/ConcertSyncClient/Source/ConcertSyncClient/Private/Replication/Manager/ReplicationManagerState.h
@@ -0,0 +1,45 @@
+// Copyright Epic Games, Inc. All Rights Reserved.
+
+#pragma once
+
+#include "Replication/IConcertClientReplicationManager.h"
+
+namespace UE::ConcertSyncClient::Replication
+{
+	class FReplicationManager;
+	
+	/**
+	 * Implements the State design pattern (see GOF) for FConcertClientReplicationManager.
+	 * Depending on the handshake state, the replication manager will react differently to the implemented IConcertClientReplicationManager functions.
+	 */
+	class FReplicationManagerState : public IConcertClientReplicationManager, public TSharedFromThis<FReplicationManagerState>
+	{
+	public:
+
+		FReplicationManagerState(FReplicationManager& Owner);
+
+	protected:
+
+		/**
+		 * Subclasses can change the state with this function.
+		 * Important: This decrements the reference count so it may destroy your state unless you keep it alive manually!
+		 */
+		void ChangeState(TSharedRef<FReplicationManagerState> NewState);
+		
+		const FReplicationManager& GetOwner() const { return Owner; }
+		FReplicationManager& GetOwner() { return Owner; }
+		
+	private:
+
+		/** Used to change the state on the owning manager. */
+		FReplicationManager& Owner;
+		
+		/**
+		 * Do any logic for entering state here instead of constructor.
+		 * 
+		 * Important to handle "recursive" calls to ChangState and also constructor does not have access to SharedThis.
+		 */
+		virtual void OnEnterState() {}
+	};
+
+}
diff --git a/Engine/Plugins/Developer/Concert/ConcertSync/ConcertSyncClient/Source/ConcertSyncClient/Private/Replication/Manager/ReplicationManagerState_Connected.cpp b/Engine/Plugins/Developer/Concert/ConcertSync/ConcertSyncClient/Source/ConcertSyncClient/Private/Replication/Manager/ReplicationManagerState_Connected.cpp
new file mode 100644
--- /dev/null
+++ b/Engine/Plugins/Developer/Concert/ConcertSync/ConcertSyncClient/Source/ConcertSyncClient/Private/Replication/Manager/ReplicationManagerState_Connected.cpp
@@ -0,0 +1,95 @@
+// Copyright Epic Games, Inc. All Rights Reserved.
+
+#include "ReplicationManagerState_Connected.h"
+
+#include "ConcertLogGlobal.h"
+#include "ConcertTransportMessages.h"
+#include "IConcertSession.h"
+#include "ReplicationManagerState_Disconnected.h"
+#include "Replication/Formats/FullObjectFormat.h"
+#include "Replication/Messages/ConcertReplicationHandshakeMessages.h"
+#include "Replication/Processing/ClientReplicationDataCollector.h"
+#include "Replication/Processing/ObjectReplicationApplierProcessor.h"
+#include "Replication/Processing/ObjectReplicationReceiver.h"
+#include "Replication/Processing/ObjectReplicationSender.h"
+
+#include "UObject/Package.h"
+#include "UObject/UObjectGlobals.h"
+
+namespace UE::ConcertSyncClient::Replication
+{
+	FReplicationManagerState_Connected::FReplicationManagerState_Connected(
+		TSharedRef<IConcertClientSession> LiveSession,
+		IConcertClientReplicationBridge* ReplicationBridge,
+		TArray<FReplicationStreamDescription> StreamDescriptions,
+		FReplicationManager& Owner
+		)
+		: FReplicationManagerState(Owner)
+		, LiveSession(LiveSession)
+		, ReplicationBridge(ReplicationBridge)
+		// TODO: Use config to determine which replication format to use
+		, ReplicationFormat(MakeShared<ConcertSyncCore::FFullObjectFormat>())
+		, ReplicationDataSource(MakeShared<FClientReplicationDataCollector>(ReplicationBridge, ReplicationFormat, StreamDescriptions))
+		, Sender(MakeShared<ConcertSyncCore::FObjectReplicationSender>(LiveSession->GetSessionServerEndpointId(), LiveSession, ReplicationDataSource))
+		, ReceivedDataCache(MakeShared<ConcertSyncCore::FObjectReplicationCache>(ReplicationFormat))
+		, Receiver(MakeShared<ConcertSyncCore::FObjectReplicationReceiver>(LiveSession, ReceivedDataCache))
+		, ReceivedReplicationQueuer(FClientReplicationDataQueuer::Make(ReplicationBridge, ReceivedDataCache))
+		, ReplicationApplier(MakeShared<FObjectReplicationApplierProcessor>(ReplicationBridge, ReplicationFormat, ReceivedReplicationQueuer))
+	{}
+
+	FReplicationManagerState_Connected::~FReplicationManagerState_Connected()
+	{
+		// Technically not needed due to AddSP but let's be nice and clean up after ourselves
+		LiveSession->OnTick().RemoveAll(this);
+	}
+
+	TFuture<FJoinReplicatedSessionResult> FReplicationManagerState_Connected::JoinReplicationSession(FJoinReplicatedSessionArgs Args)
+	{
+		UE_LOG(LogConcert, Warning, TEXT("JoinReplicationSession requested while already in a session"));
+		return MakeFulfilledPromise<FJoinReplicatedSessionResult>(FJoinReplicatedSessionResult{ EJoinReplicationErrorCode::AlreadyInSession }).GetFuture();
+	}
+
+	void FReplicationManagerState_Connected::LeaveReplicationSession()
+	{
+		LiveSession->SendCustomEvent(FConcertReplication_LeaveEvent{}, LiveSession->GetSessionServerEndpointId(), EConcertMessageFlags::ReliableOrdered);
+		ChangeState(MakeShared<FReplicationManagerState_Disconnected>(LiveSession, ReplicationBridge, GetOwner()));
+	}
+
+	void FReplicationManagerState_Connected::OnEnterState()
+	{
+		LiveSession->OnTick().AddSP(this, &FReplicationManagerState_Connected::Tick);
+	}
+
+	void FReplicationManagerState_Connected::Tick(IConcertClientSession& Session, float DeltaTime)
+	{
+		// TODO: Set this up in a config file
+		constexpr double TimeBudget = 1.0 / 60.0;
+		double TimeLeft = TimeBudget;
+
+		auto NextIndex = [this](int32 Current){ return AlternatingTickTasks.IsValidIndex(Current + 1) ? Current + 1 : 0; };
+		int32 CurrentIndex = NextTickTaskIndex;
+		const double StartTime = FPlatformTime::Seconds();
+		do
+		{
+			FTickTask Task = AlternatingTickTasks[CurrentIndex]; 
+			Invoke(Task, this, TimeLeft);
+
+			const double TotalElapsedTime = FPlatformTime::Seconds() - StartTime;
+			TimeLeft = TimeBudget - TotalElapsedTime;
+			CurrentIndex = NextIndex(CurrentIndex);
+		}
+		while (TimeLeft > 0.0 && CurrentIndex != NextTickTaskIndex);
+		NextTickTaskIndex = CurrentIndex;
+	}
+
+	void FReplicationManagerState_Connected::TickSender(float TimeBudget)
+	{
+		Sender->ProcessObjects(TimeBudget);
+	}
+
+	void FReplicationManagerState_Connected::TickReceiver(float TimeBudget)
+	{
+		ReplicationApplier->ProcessObjects(TimeBudget);
+	}
+}
+
diff --git a/Engine/Plugins/Developer/Concert/ConcertSync/ConcertSyncClient/Source/ConcertSyncClient/Private/Replication/Manager/ReplicationManagerState_Connected.h b/Engine/Plugins/Developer/Concert/ConcertSync/ConcertSyncClient/Source/ConcertSyncClient/Private/Replication/Manager/ReplicationManagerState_Connected.h
new file mode 100644
--- /dev/null
+++ b/Engine/Plugins/Developer/Concert/ConcertSync/ConcertSyncClient/Source/ConcertSyncClient/Private/Replication/Manager/ReplicationManagerState_Connected.h
@@ -0,0 +1,113 @@
+// Copyright Epic Games, Inc. All Rights Reserved.
+
+#pragma once
+
+#include "ReplicationManagerState.h"
+#include "Replication/Processing/ClientReplicationDataCollector.h"
+#include "Replication/Processing/ClientReplicationDataQueuer.h"
+
+namespace UE::ConcertSyncCore
+{
+	class FObjectReplicationReceiver;
+}
+
+namespace UE::ConcertSyncCore
+{
+	class FObjectReplicationSender;
+}
+
+class IConcertClientSession;
+
+namespace UE::ConcertSyncClient::Replication
+{
+	class FObjectReplicationApplierProcessor;
+	class FObjectReplicationSender;
+	
+	/**
+	 * State for when the client has successfully completed a replication handshake.
+	 *
+	 * Every tick this state tries to
+	 *	- collect data and sends it to the server
+	 *	- process received data and applies it
+	 */
+	class FReplicationManagerState_Connected : public FReplicationManagerState
+	{
+	public:
+
+		FReplicationManagerState_Connected(
+			TSharedRef<IConcertClientSession> LiveSession,
+			IConcertClientReplicationBridge* ReplicationBridge,
+			TArray<FReplicationStreamDescription> StreamDescriptions,
+			FReplicationManager& Owner
+			);
+		virtual ~FReplicationManagerState_Connected() override;
+
+		//~ Begin IConcertClientReplicationManager Interface
+		virtual TFuture<FJoinReplicatedSessionResult> JoinReplicationSession(FJoinReplicatedSessionArgs Args) override;
+		virtual void LeaveReplicationSession() override;
+		virtual bool CanJoin() override { return false; }
+		virtual bool IsConnectedToReplicationSession() override { return true; }
+		//~ End IConcertClientReplicationManager Interface
+
+	private:
+
+		/** Passed to FReplicationManagerState_Disconnected */
+		TSharedRef<IConcertClientSession> LiveSession;
+		/** Passed to FReplicationManagerState_Disconnected */
+		IConcertClientReplicationBridge* ReplicationBridge;
+		
+		/** The format this client will use for sending & receiving data. */
+		TSharedRef<ConcertSyncCore::IObjectReplicationFormat> ReplicationFormat;
+
+		// Sending
+		/** Used as source of replication data. */
+		TSharedRef<FClientReplicationDataCollector> ReplicationDataSource;
+		/** Sends data collected by ReplicationDataSource to the server. */
+		TSharedRef<ConcertSyncCore::FObjectReplicationSender> Sender;
+
+		// Receiving
+		/** Stores data received by Receiver until it is consumed by ReceivedReplicationQueuer. */
+		TSharedRef<ConcertSyncCore::FObjectReplicationCache> ReceivedDataCache;
+		/** Receives data from remote endpoints via message bus.  */
+		TSharedRef<ConcertSyncCore::FObjectReplicationReceiver> Receiver;
+		/** Queues data until is can be processed. */
+		TSharedRef<FClientReplicationDataQueuer> ReceivedReplicationQueuer;
+		/** Processes data from ReceivedReplicationQueuer once we tick. */
+		TSharedRef<FObjectReplicationApplierProcessor> ReplicationApplier;
+
+		using FTickTask = void(FReplicationManagerState_Connected::*)(float TimeBudget);
+		/**
+		 * We have two tasks that need to be performed each tick: Collecting data to send and applying received data.
+		 * In order that no task starves the other, we alternate which one we start with.
+		 *
+		 * Example assuming time budget of 0.1s (unrealistic time budget):
+		 * Tick 1: TickSender takes 0.8s and finishes. TickReceiver has 0.2s but cannot finish all work.
+		 * Tick 2: TickReceiver finishes work taking 0.7s. TickSender gets the remaining 0.3s.
+		 * Tick 3. TickSender starts again followed by TickSender ...
+		 */
+		const TArray<FTickTask, TInlineAllocator<2>> AlternatingTickTasks { &FReplicationManagerState_Connected::TickSender, &FReplicationManagerState_Connected::TickReceiver };
+		/** The first index to process next tick. */
+		int32 NextTickTaskIndex = 0;
+
+		//~ Begin FReplicationManagerState Interface
+		virtual void OnEnterState() override;
+		//~ End FReplicationManagerState Interface
+
+		/**
+		 * Ticks this client.
+		 * 
+		 * This processes:
+		 *  - data that is to be sent
+		 *  - data that was received
+		 *
+		 * The tasks have a time budget so that the frame rate remains stable.
+		 * It is configured in the project settings TODO: Add config
+		 */
+		void Tick(IConcertClientSession& Session, float DeltaTime);
+
+		/** Collects and sends data to the server. */
+		void TickSender(float TimeBudget);
+		/** Processes received data and serializes UObjects. */
+		void TickReceiver(float TimeBudget);
+	};
+}
diff --git a/Engine/Plugins/Developer/Concert/ConcertSync/ConcertSyncClient/Source/ConcertSyncClient/Private/Replication/Manager/ReplicationManagerState_Disconnected.cpp b/Engine/Plugins/Developer/Concert/ConcertSync/ConcertSyncClient/Source/ConcertSyncClient/Private/Replication/Manager/ReplicationManagerState_Disconnected.cpp
new file mode 100644
--- /dev/null
+++ b/Engine/Plugins/Developer/Concert/ConcertSync/ConcertSyncClient/Source/ConcertSyncClient/Private/Replication/Manager/ReplicationManagerState_Disconnected.cpp
@@ -0,0 +1,32 @@
+// Copyright Epic Games, Inc. All Rights Reserved.
+
+#include "ReplicationManagerState_Disconnected.h"
+
+#include "ConcertLogGlobal.h"
+#include "ReplicationManagerState_Handshaking.h"
+
+namespace UE::ConcertSyncClient::Replication
+{
+	FReplicationManagerState_Disconnected::FReplicationManagerState_Disconnected(
+		TSharedRef<IConcertClientSession> LiveSession,
+		IConcertClientReplicationBridge* ReplicationBridge,
+			FReplicationManager& Owner
+		)
+		: FReplicationManagerState(Owner)
+		, LiveSession(MoveTemp(LiveSession))
+		, ReplicationBridge(ReplicationBridge)
+	{}
+
+	TFuture<FJoinReplicatedSessionResult> FReplicationManagerState_Disconnected::JoinReplicationSession(FJoinReplicatedSessionArgs Args)
+	{
+		TPromise<FJoinReplicatedSessionResult> JoinPromise;
+		TFuture<FJoinReplicatedSessionResult> JoinFuture = JoinPromise.GetFuture();
+		ChangeState(MakeShared<FReplicationManagerState_Handshaking>(MoveTemp(Args), MoveTemp(JoinPromise), LiveSession, ReplicationBridge, GetOwner()));
+		return JoinFuture;
+	}
+
+	void FReplicationManagerState_Disconnected::LeaveReplicationSession()
+	{
+		UE_LOG(LogConcert, Warning, TEXT("LeaveReplicationSession: Already disconnected."));
+	}
+}
diff --git a/Engine/Plugins/Developer/Concert/ConcertSync/ConcertSyncClient/Source/ConcertSyncClient/Private/Replication/Manager/ReplicationManagerState_Disconnected.h b/Engine/Plugins/Developer/Concert/ConcertSync/ConcertSyncClient/Source/ConcertSyncClient/Private/Replication/Manager/ReplicationManagerState_Disconnected.h
new file mode 100644
--- /dev/null
+++ b/Engine/Plugins/Developer/Concert/ConcertSync/ConcertSyncClient/Source/ConcertSyncClient/Private/Replication/Manager/ReplicationManagerState_Disconnected.h
@@ -0,0 +1,37 @@
+// Copyright Epic Games, Inc. All Rights Reserved.
+
+#pragma once
+
+#include "ReplicationManagerState.h"
+#include "Replication/Processing/ClientReplicationDataCollector.h"
+
+class IConcertClientSession;
+
+namespace UE::ConcertSyncClient::Replication
+{
+	/** Initial state: waiting for a call to JoinReplicationSession to join replication session. */
+	class FReplicationManagerState_Disconnected : public FReplicationManagerState
+	{
+	public:
+
+		FReplicationManagerState_Disconnected(
+			TSharedRef<IConcertClientSession> LiveSession,
+			IConcertClientReplicationBridge* ReplicationBridge,
+			FReplicationManager& Owner
+			);
+
+		//~ Begin IConcertClientReplicationManager Interface
+		virtual TFuture<FJoinReplicatedSessionResult> JoinReplicationSession(FJoinReplicatedSessionArgs Args) override;
+		virtual void LeaveReplicationSession() override;
+		virtual bool CanJoin() override { return true; }
+		virtual bool IsConnectedToReplicationSession() override { return false; }
+		//~ End IConcertClientReplicationManager Interface
+
+	private:
+		
+		/** Passed to FReplicationManagerState_Handshaking */
+		TSharedRef<IConcertClientSession> LiveSession;
+		/** Passed to FReplicationManagerState_Handshaking */
+		IConcertClientReplicationBridge* ReplicationBridge;
+	};
+}
\ No newline at end of file
diff --git a/Engine/Plugins/Developer/Concert/ConcertSync/ConcertSyncClient/Source/ConcertSyncClient/Private/Replication/Manager/ReplicationManagerState_Handshaking.cpp b/Engine/Plugins/Developer/Concert/ConcertSync/ConcertSyncClient/Source/ConcertSyncClient/Private/Replication/Manager/ReplicationManagerState_Handshaking.cpp
new file mode 100644
--- /dev/null
+++ b/Engine/Plugins/Developer/Concert/ConcertSync/ConcertSyncClient/Source/ConcertSyncClient/Private/Replication/Manager/ReplicationManagerState_Handshaking.cpp
@@ -0,0 +1,88 @@
+// Copyright Epic Games, Inc. All Rights Reserved.
+
+#include "ReplicationManagerState_Handshaking.h"
+
+#include "IConcertSession.h"
+#include "ReplicationManagerState_Connected.h"
+#include "ReplicationManagerState_Disconnected.h"
+#include "Replication/Messages/ConcertReplicationHandshakeMessages.h"
+
+#include "Algo/Transform.h"
+
+namespace UE::ConcertSyncClient::Replication
+{
+	FReplicationManagerState_Handshaking::FReplicationManagerState_Handshaking(
+		FJoinReplicatedSessionArgs RequestArgs,
+		TPromise<FJoinReplicatedSessionResult> JoinSessionPromise,
+		TSharedRef<IConcertClientSession> LiveSession,
+		IConcertClientReplicationBridge* ReplicationBridge,
+		FReplicationManager& Owner
+		)
+		: FReplicationManagerState(Owner)
+		, RequestArgs(MoveTemp(RequestArgs))
+		, JoinSessionPromise(MoveTemp(JoinSessionPromise))
+		, LiveSession(MoveTemp(LiveSession))
+		, ReplicationBridge(ReplicationBridge)
+	{}
+
+	FReplicationManagerState_Handshaking::~FReplicationManagerState_Handshaking()
+	{
+		if (!bFulfilledPromise)
+		{
+			JoinSessionPromise.EmplaceValue(FJoinReplicatedSessionResult{ EJoinReplicationErrorCode::Cancelled });
+		}
+	}
+
+	TFuture<FJoinReplicatedSessionResult> FReplicationManagerState_Handshaking::JoinReplicationSession(FJoinReplicatedSessionArgs Args)
+	{
+		return MakeFulfilledPromise<FJoinReplicatedSessionResult>(FJoinReplicatedSessionResult{ EJoinReplicationErrorCode::AlreadyInProgress }).GetFuture();
+	}
+
+	void FReplicationManagerState_Handshaking::LeaveReplicationSession()
+	{
+		// If the server acknowledges us (later or the response is on route but not yet processed), immediately leave the session.
+		LiveSession->SendCustomEvent(FConcertReplication_LeaveEvent{}, LiveSession->GetSessionServerEndpointId(), EConcertMessageFlags::ReliableOrdered);
+		ReturnToDisconnectedState();
+	}
+
+	void FReplicationManagerState_Handshaking::OnEnterState()
+	{
+		TArray<FReplicationStreamDescription_NetPacked> Streams_NetPacked;
+		Algo::Transform(RequestArgs.Streams, Streams_NetPacked, [](const FReplicationStreamDescription& Description){ return Description.Pack(); });
+		const FConcertReplication_Join_Request Request{ Streams_NetPacked };
+		LiveSession->SendCustomRequest<FConcertReplication_Join_Request, FConcertReplication_Join_Response>(Request, LiveSession->GetSessionServerEndpointId())
+			.Next([WeakThis = TWeakPtr<FReplicationManagerState_Handshaking>(SharedThis(this))](const FConcertReplication_Join_Response& Response)
+			{
+				// If we're still in the same state, this is valid. If we've left the state, this is invalid.
+				const TSharedPtr<FReplicationManagerState_Handshaking> This = WeakThis.Pin();
+				const bool bHasLeftState = !This.IsValid();
+				if (bHasLeftState)
+				{
+					return;
+				}
+
+				if (Response.ErrorCode == EJoinReplicationErrorCode::Success)
+				{
+					This->ChangeState(
+						MakeShared<FReplicationManagerState_Connected>(This->LiveSession, This->ReplicationBridge, MoveTemp(This->RequestArgs.Streams), This->GetOwner())
+						);
+				}
+				else
+				{
+					This->ReturnToDisconnectedState();
+				}
+
+				This->bFulfilledPromise = true;
+				This->JoinSessionPromise.EmplaceValue(Response.ErrorCode, Response.DetailedErrorMessage);
+			});
+	}
+
+	void FReplicationManagerState_Handshaking::ReturnToDisconnectedState()
+	{
+		ChangeState(
+			MakeShared<FReplicationManagerState_Disconnected>(LiveSession, ReplicationBridge, GetOwner())
+			);
+	}
+}
+
+
diff --git a/Engine/Plugins/Developer/Concert/ConcertSync/ConcertSyncClient/Source/ConcertSyncClient/Private/Replication/Manager/ReplicationManagerState_Handshaking.h b/Engine/Plugins/Developer/Concert/ConcertSync/ConcertSyncClient/Source/ConcertSyncClient/Private/Replication/Manager/ReplicationManagerState_Handshaking.h
new file mode 100644
--- /dev/null
+++ b/Engine/Plugins/Developer/Concert/ConcertSync/ConcertSyncClient/Source/ConcertSyncClient/Private/Replication/Manager/ReplicationManagerState_Handshaking.h
@@ -0,0 +1,54 @@
+// Copyright Epic Games, Inc. All Rights Reserved.
+
+#pragma once
+
+#include "IConcertSessionHandler.h"
+#include "ReplicationManagerState.h"
+#include "Replication/Messages/ConcertReplicationHandshakeMessages.h"
+#include "Replication/Processing/ClientReplicationDataCollector.h"
+
+class IConcertClientSession;
+
+namespace UE::ConcertSyncClient::Replication
+{
+	/** Initial state: waiting for a call to JoinReplicationSession to join replication session. */
+	class FReplicationManagerState_Handshaking : public FReplicationManagerState
+	{
+	public:
+		
+		FReplicationManagerState_Handshaking(
+			FJoinReplicatedSessionArgs RequestArgs,
+			TPromise<FJoinReplicatedSessionResult> JoinSessionPromise,
+			TSharedRef<IConcertClientSession> LiveSession,
+			IConcertClientReplicationBridge* ReplicationBridge,
+			FReplicationManager& Owner
+			);
+		virtual ~FReplicationManagerState_Handshaking() override;
+		
+		//~ Begin IConcertClientReplicationManager Interface
+		virtual TFuture<FJoinReplicatedSessionResult> JoinReplicationSession(FJoinReplicatedSessionArgs Args) override;
+		virtual void LeaveReplicationSession() override;
+		virtual bool CanJoin() override { return false; }
+		virtual bool IsConnectedToReplicationSession() override { return false; }
+		//~ End IConcertClientReplicationManager Interface
+
+	private:
+
+		FJoinReplicatedSessionArgs RequestArgs;
+		
+		/** Promise for finishing the handshake. */
+		TPromise<FJoinReplicatedSessionResult> JoinSessionPromise;
+		bool bFulfilledPromise = false;
+		
+		/** Passed to FReplicationManagerState_Handshaking */
+		TSharedRef<IConcertClientSession> LiveSession;
+		/** Passed to FReplicationManagerState_Handshaking */
+		IConcertClientReplicationBridge* ReplicationBridge;
+
+		//~ Begin FReplicationManagerState Interface
+		virtual void OnEnterState() override;
+		//~ End FReplicationManagerState Interface
+
+		void ReturnToDisconnectedState();
+	};
+}
\ No newline at end of file
diff --git a/Engine/Plugins/Developer/Concert/ConcertSync/ConcertSyncClient/Source/ConcertSyncClient/Private/Replication/Processing/ClientReplicationDataCollector.cpp b/Engine/Plugins/Developer/Concert/ConcertSync/ConcertSyncClient/Source/ConcertSyncClient/Private/Replication/Processing/ClientReplicationDataCollector.cpp
new file mode 100644
--- /dev/null
+++ b/Engine/Plugins/Developer/Concert/ConcertSync/ConcertSyncClient/Source/ConcertSyncClient/Private/Replication/Processing/ClientReplicationDataCollector.cpp
@@ -0,0 +1,136 @@
+// Copyright Epic Games, Inc. All Rights Reserved.
+
+#include "ClientReplicationDataCollector.h"
+
+#include "Replication/Data/ReplicationStreamDescription.h"
+#include "Replication/Formats/IObjectReplicationFormat.h"
+#include "Replication/IConcertClientReplicationBridge.h"
+#include "Replication/ReplicationPropertyFilter.h"
+#include "Replication/ReplicationStreamObjectID.h"
+
+namespace UE::ConcertSyncClient::Replication
+{
+	FClientReplicationDataCollector::FClientReplicationDataCollector(
+		IConcertClientReplicationBridge* ReplicationBridge,
+		TSharedRef<ConcertSyncCore::IObjectReplicationFormat> ReplicationFormat,
+		TArrayView<FReplicationStreamDescription> StreamDescriptions
+		)
+		: Bridge(ReplicationBridge)
+		, ReplicationFormat(MoveTemp(ReplicationFormat))
+	{
+		Bridge->OnObjectDiscovered().AddRaw(this, &FClientReplicationDataCollector::StartTrackingObject);
+		Bridge->OnObjectHidden().AddRaw(this, &FClientReplicationDataCollector::StopTrackingObject);
+
+		TSet<FSoftObjectPath> PushedSet;
+		for (const FReplicationStreamDescription& StreamDescription : StreamDescriptions)
+		{
+			const FGuid StreamId = StreamDescription.BaseDescription.Identifier;
+			for (const TPair<FSoftObjectPath, FReplicatedObjectInfo>& Pair : StreamDescription.BaseDescription.ReplicationMap.ReplicatedObjects)
+			{
+				const FSoftObjectPath& ObjectPath = Pair.Key;
+				
+				TArray<FObjectInfo>& ExistingObjectInfo = ObjectsToReplicate.FindOrAdd(ObjectPath);
+				checkf(!ExistingObjectInfo.ContainsByPredicate([&StreamId](const FObjectInfo& ObjectInfo){ return ObjectInfo.StreamId == StreamId; }), TEXT("Invalid data should have been rejected during handshake"));
+
+				ExistingObjectInfo.Add({ StreamId, Pair.Value.PropertySelection });
+				// Only push it once because each push increments a counter
+				if (!PushedSet.Contains(ObjectPath))
+				{
+					PushedSet.Add(ObjectPath);
+					Bridge->PushTrackedObjects({ ObjectPath });
+				}
+			}
+		}
+	}
+
+	FClientReplicationDataCollector::~FClientReplicationDataCollector()
+	{
+		Bridge->OnObjectDiscovered().RemoveAll(this);
+		Bridge->OnObjectHidden().RemoveAll(this);
+
+		for (const TPair<FSoftObjectPath, TArray<FObjectInfo>>& Pair : ObjectsToReplicate)
+		{
+			Bridge->PopTrackedObjects({ Pair.Key });
+		}
+	}
+
+	void FClientReplicationDataCollector::ForEachPendingObject(TFunctionRef<void(const ConcertSyncCore::FReplicationStreamObjectID&)> ProcessItemFunc) const
+	{
+		for (const TPair<FSoftObjectPath, TArray<FObjectInfo>>& Pair : ObjectsToReplicate)
+		{
+			for (const FObjectInfo& ObjectInfo : Pair.Value)
+			{
+				if (ObjectInfo.ObjectCache.IsValid())
+				{
+					ProcessItemFunc({ ObjectInfo.StreamId, ObjectInfo.ObjectCache.Get() });
+				}
+			}
+		}
+	}
+
+	bool FClientReplicationDataCollector::ExtractReplicationDataForObject(
+		const ConcertSyncCore::FReplicationStreamObjectID& ObjectToProcess,
+		TFunctionRef<void(const FConcertSessionSerializedPayload& Payload)> ProcessCopyable,
+		TFunctionRef<void(FConcertSessionSerializedPayload&& Payload)> ProcessMoveable
+		)
+	{
+		const TArray<FObjectInfo>* ObjectInfos = ObjectsToReplicate.Find(ObjectToProcess.Object);
+		// ExtractReplicationDataForObject is supposed to be called in response to ForEachPendingObject... so either the call was invalid or ForEachPendingObject lied
+		if (!ensure(ObjectInfos))
+		{
+			return {};
+		}
+
+		const int32 Index = ObjectInfos->IndexOfByPredicate([&ObjectToProcess](const FObjectInfo& ObjectInfo){ return ObjectInfo.StreamId == ObjectToProcess.StreamId; });
+		// Same logic as above - either invalid call or ForEachPendingObject lied
+		if (!ensure(ObjectInfos->IsValidIndex(Index)))
+		{
+			return {};
+		}
+
+		const FObjectInfo& ObjectInfo = (*ObjectInfos)[Index];
+		// Finally... same logic as above - either invalid call or ForEachPendingObject lied
+		if (!ensure(ObjectInfo.ObjectCache.IsValid()))
+		{
+			return false;
+		}
+
+		ConcertSyncCore::FReplicationPropertyFilter Filter(ObjectInfo.SelectedProperties);
+		TOptional<FConcertSessionSerializedPayload> Payload = ReplicationFormat->CreateReplicationEvent(
+			*ObjectInfo.ObjectCache,
+			[&Filter](const FArchiveSerializedPropertyChain* Chain, const FProperty& Property)
+			{
+				return Filter.ShouldSerializeProperty(Chain, Property);
+			}
+		);
+		if (Payload)
+		{
+			ProcessMoveable(MoveTemp(*Payload));
+		}
+		return true;
+	}
+
+	void FClientReplicationDataCollector::StartTrackingObject(UObject& Object)
+	{
+		if (TArray<FObjectInfo>* ObjectInfos = ObjectsToReplicate.Find(&Object))
+		{
+			for (FObjectInfo& ObjectInfo : *ObjectInfos)
+			{
+				ObjectInfo.ObjectCache = &Object;
+			}
+			++NumTrackedObjects;
+		}
+	}
+
+	void FClientReplicationDataCollector::StopTrackingObject(UObject& Object)
+	{
+		if (TArray<FObjectInfo>* ObjectInfos = ObjectsToReplicate.Find(&Object))
+		{
+			for (FObjectInfo& ObjectInfo : *ObjectInfos)
+			{
+				ObjectInfo.ObjectCache = nullptr;
+			}
+			--NumTrackedObjects;
+		}
+	}
+}
diff --git a/Engine/Plugins/Developer/Concert/ConcertSync/ConcertSyncClient/Source/ConcertSyncClient/Private/Replication/Processing/ClientReplicationDataCollector.h b/Engine/Plugins/Developer/Concert/ConcertSync/ConcertSyncClient/Source/ConcertSyncClient/Private/Replication/Processing/ClientReplicationDataCollector.h
new file mode 100644
--- /dev/null
+++ b/Engine/Plugins/Developer/Concert/ConcertSync/ConcertSyncClient/Source/ConcertSyncClient/Private/Replication/Processing/ClientReplicationDataCollector.h
@@ -0,0 +1,75 @@
+// Copyright Epic Games, Inc. All Rights Reserved.
+
+#pragma once
+
+#include "Replication/Processing/IReplicationDataSource.h"
+
+#include "Containers/ArrayView.h"
+#include "Containers/Set.h"
+#include "Replication/Data/ConcertPropertySelection.h"
+#include "Templates/SharedPointer.h"
+#include "UObject/SoftObjectPtr.h"
+
+class FConcertSyncClientLiveSession;
+class IConcertClientReplicationBridge;
+class UObject;
+struct FConcertPropertySelection;
+struct FReplicationStreamDescription;
+
+namespace UE::ConcertSyncCore
+{
+	class IObjectReplicationFormat;
+	class FObjectReplicationProcessor;
+}
+
+namespace UE::ConcertSyncClient::Replication
+{
+	
+	/**
+	 * Exposes UObject instances to an FObjectReplicationProcessor.
+	 * IConcertClientReplicationBridge tracks UObject lifetime, this class exposes them.
+	 */
+	class FClientReplicationDataCollector : public ConcertSyncCore::IReplicationDataSource
+	{
+	public:
+		
+		FClientReplicationDataCollector(
+			IConcertClientReplicationBridge* ReplicationBridge,
+			TSharedRef<ConcertSyncCore::IObjectReplicationFormat> ReplicationFormat,
+			TArrayView<FReplicationStreamDescription> StreamsToSend
+			);
+		virtual ~FClientReplicationDataCollector() override;
+
+		//~ Begin IReplicationDataSource Interface
+		virtual void ForEachPendingObject(TFunctionRef<void(const ConcertSyncCore::FReplicationStreamObjectID&)> ProcessItemFunc) const override;
+		virtual int32 NumObjects() const override { return NumTrackedObjects; }
+		virtual bool ExtractReplicationDataForObject(const ConcertSyncCore::FReplicationStreamObjectID& Object, TFunctionRef<void(const FConcertSessionSerializedPayload& Payload)> ProcessCopyable, TFunctionRef<void(FConcertSessionSerializedPayload&& Payload)> ProcessMoveable) override;
+		//~ End IReplicationDataSource Interface
+
+	private:
+
+		/** Gets and tracks replicated objects */
+		IConcertClientReplicationBridge* Bridge;
+		/** Used to create the replication data sent to the server. */
+		TSharedRef<ConcertSyncCore::IObjectReplicationFormat> ReplicationFormat;
+
+		struct FObjectInfo
+		{
+			/** The replication stream producing this object's data */
+			FGuid StreamId;
+			/** The properties to replicate */
+			FConcertPropertySelection SelectedProperties;
+			/** Set when the bridge tells us the object is available for replication. Unset otherwise. */
+			TWeakObjectPtr<UObject> ObjectCache;
+		};
+		
+		/** The objects and their properties to replicate */
+		TMap<FSoftObjectPath, TArray<FObjectInfo>> ObjectsToReplicate;
+		/** Cached number of FObjectInfo in ObjectsToReplicate that have a valid FObjectInfo::ObjectCache. */
+		int32 NumTrackedObjects = 0;
+
+		// Handle events from bridge
+		void StartTrackingObject(UObject& Object);
+		void StopTrackingObject(UObject& Object);
+	};
+}
diff --git a/Engine/Plugins/Developer/Concert/ConcertSync/ConcertSyncClient/Source/ConcertSyncClient/Private/Replication/Processing/ClientReplicationDataQueuer.cpp b/Engine/Plugins/Developer/Concert/ConcertSync/ConcertSyncClient/Source/ConcertSyncClient/Private/Replication/Processing/ClientReplicationDataQueuer.cpp
new file mode 100644
--- /dev/null
+++ b/Engine/Plugins/Developer/Concert/ConcertSync/ConcertSyncClient/Source/ConcertSyncClient/Private/Replication/Processing/ClientReplicationDataQueuer.cpp
@@ -0,0 +1,30 @@
+// Copyright Epic Games, Inc. All Rights Reserved.
+
+#include "ClientReplicationDataQueuer.h"
+
+#include "Replication/IConcertClientReplicationBridge.h"
+
+namespace UE::ConcertSyncClient::Replication
+{
+	TSharedRef<FClientReplicationDataQueuer> FClientReplicationDataQueuer::Make(
+		IConcertClientReplicationBridge* ReplicationBridge,
+		TSharedRef<ConcertSyncCore::FObjectReplicationCache> InReplicationCache
+		)
+	{
+		TSharedRef<FClientReplicationDataQueuer> Result = MakeShared<FClientReplicationDataQueuer>(ReplicationBridge);
+		Result->BindToCache(MoveTemp(InReplicationCache));
+		return Result;
+	}
+	
+	FClientReplicationDataQueuer::FClientReplicationDataQueuer(IConcertClientReplicationBridge* ReplicationBridge)
+		: ReplicationBridge(ReplicationBridge)
+	{}
+
+	bool FClientReplicationDataQueuer::WantsToAcceptObject(const ConcertSyncCore::FReplicationStreamObjectID& Object) const
+	{
+		// The server tries not to send objects to are unavailable but it could happen the object becomes unavailable while data is underway
+		// TODO: Check whether client attributes allow this object since they could have changed since
+		const bool bIsAvailable = ReplicationBridge->IsObjectAvailable(Object.Object);
+		return bIsAvailable;
+	}
+}
diff --git a/Engine/Plugins/Developer/Concert/ConcertSync/ConcertSyncClient/Source/ConcertSyncClient/Private/Replication/Processing/ClientReplicationDataQueuer.h b/Engine/Plugins/Developer/Concert/ConcertSync/ConcertSyncClient/Source/ConcertSyncClient/Private/Replication/Processing/ClientReplicationDataQueuer.h
new file mode 100644
--- /dev/null
+++ b/Engine/Plugins/Developer/Concert/ConcertSync/ConcertSyncClient/Source/ConcertSyncClient/Private/Replication/Processing/ClientReplicationDataQueuer.h
@@ -0,0 +1,34 @@
+// Copyright Epic Games, Inc. All Rights Reserved.
+
+#pragma once
+
+#include "Replication/Processing/ReplicationDataQueuer.h"
+#include "Templates/SharedPointer.h"
+
+class IConcertClientReplicationBridge;
+
+namespace UE::ConcertSyncClient::Replication
+{
+
+	/** Queues events received from the server. */
+	class FClientReplicationDataQueuer : public ConcertSyncCore::FReplicationDataQueuer
+	{
+		template <typename ObjectType, ESPMode Mode>
+		friend class SharedPointerInternals::TIntrusiveReferenceController;
+	public:
+
+		static TSharedRef<FClientReplicationDataQueuer> Make(IConcertClientReplicationBridge* ReplicationBridge, TSharedRef<ConcertSyncCore::FObjectReplicationCache> InReplicationCache);
+		
+		//~ Begin IReplicationCacheUser Interface
+		virtual bool WantsToAcceptObject(const ConcertSyncCore::FReplicationStreamObjectID& Object) const override;
+		//~ End IReplicationCacheUser Interface
+
+	private:
+		
+		FClientReplicationDataQueuer(IConcertClientReplicationBridge* ReplicationBridge);
+
+		/** Used to determine whether the object is currently alive. */
+		IConcertClientReplicationBridge* ReplicationBridge;
+	};
+}
+
diff --git a/Engine/Plugins/Developer/Concert/ConcertSync/ConcertSyncClient/Source/ConcertSyncClient/Private/Replication/Processing/ObjectReplicationApplierProcessor.cpp b/Engine/Plugins/Developer/Concert/ConcertSync/ConcertSyncClient/Source/ConcertSyncClient/Private/Replication/Processing/ObjectReplicationApplierProcessor.cpp
new file mode 100644
--- /dev/null
+++ b/Engine/Plugins/Developer/Concert/ConcertSync/ConcertSyncClient/Source/ConcertSyncClient/Private/Replication/Processing/ObjectReplicationApplierProcessor.cpp
@@ -0,0 +1,47 @@
+// Copyright Epic Games, Inc. All Rights Reserved.
+
+#include "ObjectReplicationApplierProcessor.h"
+
+#include "ConcertLogGlobal.h"
+#include "Replication/IConcertClientReplicationBridge.h"
+#include "Replication/Formats/IObjectReplicationFormat.h"
+
+#include "Components/SceneComponent.h"
+
+namespace UE::ConcertSyncClient::Replication
+{
+	FObjectReplicationApplierProcessor::FObjectReplicationApplierProcessor(
+		IConcertClientReplicationBridge* ReplicationBridge,
+		TSharedRef<ConcertSyncCore::IObjectReplicationFormat> ReplicationFormat,
+		TSharedRef<ConcertSyncCore::IReplicationDataSource> DataSource
+		)
+		: FObjectReplicationProcessor(MoveTemp(DataSource))
+		, ReplicationBridge(ReplicationBridge)
+		, ReplicationFormat(MoveTemp(ReplicationFormat))
+	{}
+
+	void FObjectReplicationApplierProcessor::ProcessObject(const FObjectProcessArgs& Args)
+	{
+		UObject* Object = ReplicationBridge->FindObjectIfAvailable(Args.ObjectInfo.Object);
+		if (!Object)
+		{
+			UE_LOG(LogConcert, Error, TEXT("Replication: Object %s is unavailable. The data source should not have reported it."), *Args.ObjectInfo.Object.ToString());
+			return;
+		}
+
+		bool bAppliedData = false;
+		GetDataSource().ExtractReplicationDataForObject(Args.ObjectInfo, [this, Object, &bAppliedData](const FConcertSessionSerializedPayload& Payload)
+		{
+			bAppliedData = true;
+			ReplicationFormat->ApplyReplicationEvent(*Object, Payload);
+
+			// TODO: This is very hacky and leaves performance on the table... this is in case ApplyReplicationEvent updates the transform
+			if (USceneComponent* SceneComponent = Cast<USceneComponent>(Object))
+			{
+				SceneComponent->UpdateComponentToWorld();
+			}
+		});
+		// This should not happen. If it does, we're wasting  network bandwidth.
+		UE_CLOG(!bAppliedData, LogConcert, Warning, TEXT("Replication: Server sent data that could not be applied (likely it was empty) for object %s from stream %s"), *Args.ObjectInfo.Object.ToString(), *Args.ObjectInfo.StreamId.ToString());
+	}
+}
diff --git a/Engine/Plugins/Developer/Concert/ConcertSync/ConcertSyncClient/Source/ConcertSyncClient/Private/Replication/Processing/ObjectReplicationApplierProcessor.h b/Engine/Plugins/Developer/Concert/ConcertSync/ConcertSyncClient/Source/ConcertSyncClient/Private/Replication/Processing/ObjectReplicationApplierProcessor.h
new file mode 100644
--- /dev/null
+++ b/Engine/Plugins/Developer/Concert/ConcertSync/ConcertSyncClient/Source/ConcertSyncClient/Private/Replication/Processing/ObjectReplicationApplierProcessor.h
@@ -0,0 +1,38 @@
+// Copyright Epic Games, Inc. All Rights Reserved.
+
+#pragma once
+
+#include "Replication/Processing/ObjectReplicationProcessor.h"
+#include "Replication/Processing/ReplicationDataQueuer.h"
+
+class IConcertClientReplicationBridge;
+
+namespace UE::ConcertSyncClient::Replication
+{
+	/** Applies replication data to the object is is meant for. */
+	class FObjectReplicationApplierProcessor : public ConcertSyncCore::FObjectReplicationProcessor 
+	{
+	public:
+
+		FObjectReplicationApplierProcessor(
+			IConcertClientReplicationBridge* ReplicationBridge,
+			TSharedRef<ConcertSyncCore::IObjectReplicationFormat> ReplicationFormat,
+			TSharedRef<ConcertSyncCore::IReplicationDataSource> DataSource
+			);
+		
+	protected:
+
+		//~ Begin FObjectReplicationProcessor Interface
+		virtual void ProcessObject(const FObjectProcessArgs& Args) override;
+		//~ End FObjectReplicationProcessor Interface
+
+	private:
+
+		/** Retrieves objects for applying replication data to. */
+		IConcertClientReplicationBridge* ReplicationBridge;
+
+		/** Unpacks replication data and applies it to an UObject instance. */
+		TSharedRef<ConcertSyncCore::IObjectReplicationFormat> ReplicationFormat;
+	};
+}
+
diff --git a/Engine/Plugins/Developer/Concert/ConcertSync/ConcertSyncClient/Source/ConcertSyncClient/Private/Replication/ReplicationTestInterface.cpp b/Engine/Plugins/Developer/Concert/ConcertSync/ConcertSyncClient/Source/ConcertSyncClient/Private/Replication/ReplicationTestInterface.cpp
new file mode 100644
--- /dev/null
+++ b/Engine/Plugins/Developer/Concert/ConcertSync/ConcertSyncClient/Source/ConcertSyncClient/Private/Replication/ReplicationTestInterface.cpp
@@ -0,0 +1,24 @@
+// Copyright Epic Games, Inc. All Rights Reserved.
+
+#include "Replication/IConcertClientReplicationBridge.h"
+#include "Replication/IConcertClientReplicationManager.h"
+#include "Manager/ReplicationManager.h"
+
+/**
+ * Exposes functions that are required for testing.
+ * 
+ * These functions are technically exported but conceptually not part of the public interface
+ * and should only be used for the purpose of automated testing.
+ */
+namespace UE::ConcertSyncClient::TestInterface
+{
+	CONCERTSYNCCLIENT_API TSharedRef<IConcertClientReplicationManager> CreateClientReplicationManager(
+		TSharedRef<IConcertClientSession> InLiveSession,
+		IConcertClientReplicationBridge* InBridge
+		)
+	{
+		const TSharedRef<Replication::FReplicationManager> Result = MakeShared<Replication::FReplicationManager>(MoveTemp(InLiveSession), InBridge);
+		Result->StartAcceptingJoinRequests();
+		return Result;
+	}
+}
diff --git a/Engine/Plugins/Developer/Concert/ConcertSync/ConcertSyncClient/Source/ConcertSyncClient/Public/IConcertSyncClient.h b/Engine/Plugins/Developer/Concert/ConcertSync/ConcertSyncClient/Source/ConcertSyncClient/Public/IConcertSyncClient.h
--- a/Engine/Plugins/Developer/Concert/ConcertSync/ConcertSyncClient/Source/ConcertSyncClient/Public/IConcertSyncClient.h
+++ b/Engine/Plugins/Developer/Concert/ConcertSync/ConcertSyncClient/Source/ConcertSyncClient/Public/IConcertSyncClient.h
@@ -9,15 +9,17 @@
 #include "ConcertActionDefinition.h"
 #include "ConcertSyncSessionFlags.h"
 
-class UConcertClientConfig;
-class IConcertClientWorkspace;
+class IConcertFileSharingService;
+class IConcertClientPackageBridge;
 class IConcertClientPresenceManager;
+class IConcertClientReplicationManager;
+class IConcertClientReplicationBridge;
 class IConcertClientSequencerManager;
-struct FConcertSessionClientInfo;
 class IConcertClientTransactionBridge;
-class IConcertClientPackageBridge;
+class IConcertClientWorkspace;
 class IConcertSyncClient;
-class IConcertFileSharingService;
+class UConcertClientConfig;
+struct FConcertSessionClientInfo;
 
 DECLARE_MULTICAST_DELEGATE_OneParam(FOnConcertClientWorkspaceStartupOrShutdown, const TSharedPtr<IConcertClientWorkspace>& /** InClientWorkspace */ );
 DECLARE_MULTICAST_DELEGATE_OneParam(FOnConcertClientSyncSessionStartupOrShutdown, const IConcertSyncClient* /* InSyncClient */ );
@@ -55,6 +57,8 @@ public:
 	/** Get the current session client workspace, if any. */
 	virtual TSharedPtr<IConcertClientWorkspace> GetWorkspace() const = 0;
 
+
+	
 	/** 
 	 * Get the current session presence manager, if any.
 	 * @note that pointer shouldn't be stored and always accessed through this client.
@@ -67,6 +71,14 @@ public:
 	 */
 	virtual IConcertClientSequencerManager* GetSequencerManager() const = 0;
 
+	/**
+	 * Get the current session replication manager, if any.
+	 * @note that pointer shouldn't be stored and always accessed through this client.
+	 */
+	virtual IConcertClientReplicationManager* GetReplicationManager() const = 0;
+
+
+	
 	/** Get the delegate called on every workspace startup. */
 	virtual FOnConcertClientWorkspaceStartupOrShutdown& OnWorkspaceStartup() = 0;
 
@@ -79,6 +91,8 @@ public:
 	/** Get the delegate called just after the workspace gets deleted. */
 	virtual FOnConcertClientSyncSessionStartupOrShutdown& OnSyncSessionShutdown() = 0;
 
+
+	
 	/** Persist specific packages */
 	virtual void PersistSpecificChanges(TArrayView<const FName> Packages) = 0;
 
@@ -96,9 +110,14 @@ public:
 	/** Set the file sharing service, enabling the server to work with large files. The client and the server must use compatible sharing services. The service must be set prior establishing a connection. */
 	virtual void SetFileSharingService(TSharedPtr<IConcertFileSharingService> InFileSharingService) = 0;
 
+
+	
 	/** Get the current transaction bridge for this client */
 	virtual IConcertClientTransactionBridge* GetTransactionBridge() const = 0;
 
 	/** Get the current package bridge for this client. */
 	virtual IConcertClientPackageBridge* GetPackageBridge() const = 0;
+
+	/** Get the current replication bridge for this client. */
+	virtual IConcertClientReplicationBridge* GetReplicationBridge() const = 0;
 };
diff --git a/Engine/Plugins/Developer/Concert/ConcertSync/ConcertSyncClient/Source/ConcertSyncClient/Public/IConcertSyncClientModule.h b/Engine/Plugins/Developer/Concert/ConcertSync/ConcertSyncClient/Source/ConcertSyncClient/Public/IConcertSyncClientModule.h
--- a/Engine/Plugins/Developer/Concert/ConcertSync/ConcertSyncClient/Source/ConcertSyncClient/Public/IConcertSyncClientModule.h
+++ b/Engine/Plugins/Developer/Concert/ConcertSync/ConcertSyncClient/Source/ConcertSyncClient/Public/IConcertSyncClientModule.h
@@ -6,13 +6,16 @@
 #include "Modules/ModuleInterface.h"
 #include "Modules/ModuleManager.h"
 
-class IConcertSyncClient;
 class IConcertClientPackageBridge;
+class IConcertClientReplicationBridge;
 class IConcertClientTransactionBridge;
+class IConcertSyncClient;
 class UConcertClientConfig;
 
 DECLARE_MULTICAST_DELEGATE_OneParam(FOnConcertClientCreated, TSharedRef<IConcertSyncClient>);
 
+class IConcertClientReplicationBridge;
+
 /**
  * Interface for the Concert Sync Client module.
  */
@@ -56,28 +59,21 @@ public:
 	 */
 	virtual TSharedRef<IConcertSyncClient> CreateClient(const FString& InRole) = 0;
 
-	/**
-	 * Get the bridge between the editor package system and Concert.
-	 */
+	/** Get the bridge between the editor package system and Concert. */
 	virtual IConcertClientPackageBridge& GetPackageBridge() = 0;
 
-	/**
-	 * Get the bridge between the editor transaction system and Concert.
-	 */
+	/** Get the bridge between the editor transaction system and Concert. */
 	virtual IConcertClientTransactionBridge& GetTransactionBridge() = 0;
 
-	/**
-	 * Returns the list of active clients.
-	 */
+	/** Gets the bridge between the editor UObjects and Concert. */
+	virtual IConcertClientReplicationBridge& GetReplicationBridge() = 0;
+
+	/** Returns the list of active clients. */
 	virtual TArray<TSharedRef<IConcertSyncClient>> GetClients() const = 0;
 
-	/**
-	 * Find a concert sync client.
-	 */
+	/** Find a concert sync client. */
 	virtual TSharedPtr<IConcertSyncClient> GetClient(const FString& InRole) const = 0;
 
-	/**
-	 * Delegate invoked when a new concert sync client is created.
-	 */
+	/** Delegate invoked when a new concert sync client is created. */
 	virtual FOnConcertClientCreated& OnClientCreated() = 0;
 };
diff --git a/Engine/Plugins/Developer/Concert/ConcertSync/ConcertSyncClient/Source/ConcertSyncClient/Public/Replication/IConcertClientReplicationBridge.h b/Engine/Plugins/Developer/Concert/ConcertSync/ConcertSyncClient/Source/ConcertSyncClient/Public/Replication/IConcertClientReplicationBridge.h
new file mode 100644
--- /dev/null
+++ b/Engine/Plugins/Developer/Concert/ConcertSync/ConcertSyncClient/Source/ConcertSyncClient/Public/Replication/IConcertClientReplicationBridge.h
@@ -0,0 +1,81 @@
+// Copyright Epic Games, Inc. All Rights Reserved.
+
+#pragma once
+
+#include "Containers/ArrayView.h"
+#include "Delegates/Delegate.h"
+#include "Delegates/DelegateCombinations.h"
+
+class UObject;
+struct FConcertObjectReplicationEvent;
+struct FSoftObjectPath;
+
+DECLARE_MULTICAST_DELEGATE_OneParam(FConcertClientReplicationBridgeObjectEvent, UObject&);
+
+/**
+ * A session-independent object bridging Unreal Engine with Concert.
+ * Responsible for tracking the lifetime of replicated objects, i.e. level (un)load, etc.
+ *
+ * An UObject is called "available" if it is in a state suitable for replicating; otherwise it is "unavailable".
+ * It is "tracked" if this bridge is keep track of whether is available; otherwise is is "untracked".
+ * The moment a tracked object is determined to be available it is said to "have been discovered". The moment is
+ * determined to be unavailable, it is said to "have become hidden".
+ * 
+ * In the editor an UObject should be considered
+ *	- available:
+ *		- in an UWorld the client currently has open,
+ *		- in an active data layer,
+ *	- unavailable 
+ *		- in a non-loaded UWorld,
+ *		- in an unloaded or loaded data layer (remember that loaded means not visible),
+ *		- does not have an UWorld as outer (TODO: in the future this restriction could be lifted - there is no technical reason for it - VP workflows just do not need it right now)
+ */
+class IConcertClientReplicationBridge
+{
+public:
+
+	virtual ~IConcertClientReplicationBridge() = default;
+
+	/**
+	 * Starts tracking a list of objects, possibly triggering OnObjectDiscovered immediately.
+	 * 
+	 * Increments a tracking counter for how often this object was requested to be tracked;
+	 * there may be multiple independent systems requiring tracking.
+	 */
+	virtual void PushTrackedObjects(TArrayView<const FSoftObjectPath> TrackedObjects) = 0;
+	/**
+	 * Stops tracking a list of objects, possibly triggering OnObjectRemoved immediately.
+	 * Decrements the tracking counter. If it reaches 0, the object is no longer tracked;
+	 * there may be multiple independent systems requiring tracking.
+	 */
+	virtual void PopTrackedObjects(TArrayView<const FSoftObjectPath> TrackedObjects) = 0;
+
+	/**
+	 * More lightweight version of FindObjectIfAvailable which does not force the UObject to be loaded, if it is
+	 * determined to be available.
+	 * This function should be called instead of FindAndTrackObjectIfAvailable if the UObject instance is not needed.
+	 *
+	 * This function will not discover nor hide any UObject (that means broadcast OnObjectDiscovered or OnObjectHidden),
+	 * unless it is tracked, in which case it can.
+	 * 
+	 * @param Path The object to find
+	 * @return Whether a call to FindObjectIfAvailable would return a valid object.
+	 */
+	virtual bool IsObjectAvailable(const FSoftObjectPath& Path) = 0;
+
+	/**
+	 * Obtains the object if it is available. If the UObject is unloaded, it shall be loaded.
+	 * 
+	 * This function will not discover nor hide any UObject (that means broadcast OnObjectDiscovered or OnObjectHidden),
+	 * unless it is tracked, in which case it can.
+	 * 
+	 * @param Path The object to find
+	 * @return The object if it was valid or nullptr
+	 */
+	virtual UObject* FindObjectIfAvailable(const FSoftObjectPath& Path) = 0;
+	
+	/** Event for when a replicatable object becomes available, e.g. because its level is loaded. */
+	virtual FConcertClientReplicationBridgeObjectEvent& OnObjectDiscovered() = 0;
+	/** Event for when an object that was previously discovered becomes unavailable, e.g. because it is unloaded. */
+	virtual FConcertClientReplicationBridgeObjectEvent& OnObjectHidden() = 0;
+};
diff --git a/Engine/Plugins/Developer/Concert/ConcertSync/ConcertSyncClient/Source/ConcertSyncClient/Public/Replication/IConcertClientReplicationManager.h b/Engine/Plugins/Developer/Concert/ConcertSync/ConcertSyncClient/Source/ConcertSyncClient/Public/Replication/IConcertClientReplicationManager.h
new file mode 100644
--- /dev/null
+++ b/Engine/Plugins/Developer/Concert/ConcertSync/ConcertSyncClient/Source/ConcertSyncClient/Public/Replication/IConcertClientReplicationManager.h
@@ -0,0 +1,69 @@
+// Copyright Epic Games, Inc. All Rights Reserved.
+
+#pragma once
+
+#include "Replication/Messages/ConcertReplicationHandshakeMessages.h"
+#include "Replication/Data/ReplicationClientDescription.h"
+
+template<typename ResultType>
+class TFuture;
+
+struct FReplicationClientDescription;
+struct FReplicationStreamDescription;
+
+namespace UE::ConcertSyncClient::Replication
+{
+	struct FJoinReplicatedSessionArgs
+	{
+		/** General info about this client, such as the type of data it wishes to receive. */
+		FReplicationClientDescription ClientInfo;
+		/** The streams this client offers. */
+		TArray<FReplicationStreamDescription> Streams;
+	};
+
+	struct FJoinReplicatedSessionResult
+	{
+		/** Error code sent by the server. */
+		EJoinReplicationErrorCode ErrorCode;
+		/** Optional error message to help human user resolve the error. */
+		FString DetailedErrorMessage;
+
+		FJoinReplicatedSessionResult(EJoinReplicationErrorCode ErrorCode, FString DetailedErrorMessage = TEXT(""))
+			: ErrorCode(ErrorCode)
+			, DetailedErrorMessage(MoveTemp(DetailedErrorMessage))
+		{}
+	};
+}
+
+/**
+ * Handles all communication with the server regarding replication.
+ * 
+ * Keeps a list of properties to send along with their send rules.
+ * Tells the server which properties this client is interested in receiving.
+ */
+class IConcertClientReplicationManager
+{
+public:
+
+	/**
+	 * Joins a replication session.
+	 * Subsequent calls to JoinReplicationSession will fail until either the resulting TFuture returns or LeaveReplicationSession is called.
+	 *
+	 * @note This may execute on any thread. Take care to synchronize correctly with the game thread if needed.
+	 */
+	virtual TFuture<UE::ConcertSyncClient::Replication::FJoinReplicatedSessionResult> JoinReplicationSession(UE::ConcertSyncClient::Replication::FJoinReplicatedSessionArgs Args) = 0;
+
+	/** Leaves the current replication session. */
+	virtual void LeaveReplicationSession() = 0;
+
+	/**
+	 * Whether it is valid to call JoinReplicationSession right now. Returns false if it was called and its future has not yet concluded.
+	 * If this returns true, you can call JoinReplicationSession to attempt joining the session.
+	 */
+	virtual bool CanJoin() = 0;
+		
+	/** Whether JoinReplicationSession completed successfully and LeaveReplicationSession has not yet been called. */
+	virtual bool IsConnectedToReplicationSession() = 0;
+		
+	virtual ~IConcertClientReplicationManager() = default;
+};
\ No newline at end of file
diff --git a/Engine/Plugins/Developer/Concert/ConcertSync/ConcertSyncCore/Source/ConcertSyncCore/Private/Replication/Data/ConcertPropertySelection.cpp b/Engine/Plugins/Developer/Concert/ConcertSync/ConcertSyncCore/Source/ConcertSyncCore/Private/Replication/Data/ConcertPropertySelection.cpp
--- a/Engine/Plugins/Developer/Concert/ConcertSync/ConcertSyncCore/Source/ConcertSyncCore/Private/Replication/Data/ConcertPropertySelection.cpp
+++ b/Engine/Plugins/Developer/Concert/ConcertSync/ConcertSyncCore/Source/ConcertSyncCore/Private/Replication/Data/ConcertPropertySelection.cpp
@@ -77,7 +77,7 @@ FConcertPropertyChain::FConcertPropertyChain(const FArchiveSerializedPropertyCha
 		PathToProperty.Add(LeafProperty.GetFName());
 	}
 
-	UE_CLOG(!IsReplicatableProperty(LeafProperty), LogConcert, Warning, TEXT("Instantiated property chain with non replicable property: %s"), *ToString());
+	UE_CLOG(!IsReplicatableProperty(LeafProperty), LogConcert, Warning, TEXT("Instantiated property chain with non replicatable property: %s"), *ToString());
 }
 
 bool FConcertPropertyChain::IsChildOf(const FConcertPropertyChain& ParentToCheck) const
@@ -113,6 +113,71 @@ bool FConcertPropertyChain::IsDirectChildOf(const FConcertPropertyChain& ParentT
 	return bArePathsEqual;
 }
 
+bool FConcertPropertyChain::MatchesExactly(const FArchiveSerializedPropertyChain* OptionalChain, const FProperty& LeafProperty) const
+{
+	using namespace UE::ConcertSyncCore::PropertyChain;
+	const int32 OptionalChainLength = OptionalChain ? OptionalChain->GetNumProperties() : 0;
+	
+	const bool bLeafPropertiesMatch = !PathToProperty.IsEmpty() && IsInnerContainerProperty(LeafProperty)
+		// The only place FConcertPropertyChain contains inner container properties is at the end; it is named InternalContainerPropertyValueName.
+		// In that case the "real" leaf property is the owning container property.
+		? OptionalChainLength >= 1 && OptionalChain->GetPropertyFromStack(0)->GetFName() == PathToProperty[PathToProperty.Num() - 2] 
+		: LeafProperty.GetFName() == PathToProperty[PathToProperty.Num() - 1];
+	// No point continuing matching properties if leaf properties do not match.
+	if (!bLeafPropertiesMatch)
+	{
+		return false;
+	}
+
+	// Does the chain describe a single property? If yes, we're done.
+	if (OptionalChainLength == 0 && PathToProperty.Num() == 1)
+	{
+		return true;
+	}
+
+	// We get here if PathToproperty.Num() > 1 so OptionalChain must be valid, otherwise we're not equal.
+	if (!OptionalChain)
+	{
+		return false;
+	}
+
+	/*
+	 * OptionalChainLength can be larger than PathToProperty.Num() - 1 and still match!
+	 * 
+	 * Example:
+	 *	USTRUCT() struct FInner { UPROPERTY() float Value; }
+	 *	UCLASS() class UFoo : public UObject { UPROPERTY() TArray<FInner> Nested; }
+	 *	
+	 * Supposing PathToProperty = { "Nested", "Value" }, this is what the input to Matches would be:
+	 *  - OptionalChain = { "Nested", "Nested" }
+	 *  - LeafProperty = "Value"
+	 * The 1st "Nested" is FArrayProperty and the 2nd "Nested" is the FArrayProperty::Inner property, which is always named the same way.
+	 */	
+	bool bArePathsEqual = true;
+
+	int32 Index_PathToProperty = 0;
+	int32 Index_OptionalChain = 0;
+	// Advance Index_OptionalChain towards the end of OptionalChain...
+	for (; bArePathsEqual && PathToProperty.IsValidIndex(Index_PathToProperty) && Index_OptionalChain < OptionalChain->GetNumProperties(); ++Index_OptionalChain)
+	{
+		const FProperty* OptionalChainProperty = OptionalChain->GetPropertyFromRoot(Index_OptionalChain);
+		const bool bIsInternalContainerProperty = IsInnerContainerProperty(*OptionalChainProperty);
+		bArePathsEqual &= bIsInternalContainerProperty || OptionalChainProperty->GetFName() == PathToProperty[Index_PathToProperty];
+
+		// ... but only advance Index_PathToProperty if the current property is not an inner container property since they do not show up in FConcertPropertyChains.
+		if (!bIsInternalContainerProperty)
+		{
+			++Index_PathToProperty;
+		}
+	}
+
+	// If not every property was visited, that could mean one path is is a sub-path of the other. Regardless of the reason, they stopped matching somewhere in the middle so return false.
+	const bool bVisitedEveryProperty = Index_OptionalChain == OptionalChain->GetNumProperties()
+		// -1 to exclude leaf property (we checked it way at the beginning): OptionalChain does not contain the leaf property.
+		&& (Index_PathToProperty == PathToProperty.Num() - 1);
+	return bArePathsEqual && bVisitedEveryProperty;
+}
+
 FString FConcertPropertyChain::ToString(EToStringMethod Method) const
 {
 	switch (Method)
diff --git a/Engine/Plugins/Developer/Concert/ConcertSync/ConcertSyncCore/Source/ConcertSyncCore/Private/Replication/Formats/FullObjectFormat.cpp b/Engine/Plugins/Developer/Concert/ConcertSync/ConcertSyncCore/Source/ConcertSyncCore/Private/Replication/Formats/FullObjectFormat.cpp
new file mode 100644
--- /dev/null
+++ b/Engine/Plugins/Developer/Concert/ConcertSync/ConcertSyncCore/Source/ConcertSyncCore/Private/Replication/Formats/FullObjectFormat.cpp
@@ -0,0 +1,94 @@
+// Copyright Epic Games, Inc. All Rights Reserved.
+
+#include "Replication/Formats/FullObjectFormat.h"
+
+#include "ConcertMessageData.h"
+#include "Replication/Formats/FullObjectReplicationData.h"
+#include "Serialization/ObjectAndNameAsStringProxyArchive.h"
+#include "Serialization/ObjectReader.h"
+#include "Serialization/ObjectWriter.h"
+
+namespace UE::ConcertSyncCore
+{
+	namespace FullObjectFormat
+	{
+		static void ApplySharedArchiveConfig(FArchive& Archive, bool bIsLoading)
+		{
+			Archive.ArIgnoreArchetypeRef = true;
+			Archive.ArIgnoreClassRef = true;
+			Archive.ArNoDelta = false;
+			
+			Archive.SetIsSaving(!bIsLoading);
+			Archive.SetIsLoading(bIsLoading);
+			Archive.SetIsPersistent(false);
+			// TODO: Test
+			//Archive.SetWantBinaryPropertySerialization(true);
+		}
+		
+		class FFullObjectFormatWriter : public FObjectWriter
+		{
+			const FFullObjectFormat::FAllowPropertyFunc& IsPropertyAllowedFunc;
+		public:
+
+			FFullObjectFormatWriter(TArray<uint8>& Data, const FFullObjectFormat::FAllowPropertyFunc& IsPropertyAllowedFunc)
+				: FObjectWriter(Data)
+				, IsPropertyAllowedFunc(IsPropertyAllowedFunc)
+			{
+				constexpr bool bIsLoading = false;
+				ApplySharedArchiveConfig(*this, bIsLoading);
+			}
+
+			virtual bool ShouldSkipProperty(const FProperty* InProperty) const override
+			{
+				const bool bIsAllowed = IsPropertyAllowedFunc(GetSerializedPropertyChain(), *InProperty);
+				const bool bSkip = !bIsAllowed;
+				return bSkip;
+			}
+		};
+
+		class FFulObjectFormatReader : public FObjectReader
+		{
+		public:
+
+			FFulObjectFormatReader(TArray<uint8>& Data)
+				: FObjectReader(Data)
+			{
+				constexpr bool bIsLoading = true;
+				ApplySharedArchiveConfig(*this, bIsLoading);
+			}
+		};
+	}
+	
+	TOptional<FConcertSessionSerializedPayload> FFullObjectFormat::CreateReplicationEvent(UObject& Object, FAllowPropertyFunc IsPropertyAllowedFunc)
+	{
+		FFullObjectReplicationData Data;
+		FullObjectFormat::FFullObjectFormatWriter Writer(Data.SerializedObjectData.Bytes, IsPropertyAllowedFunc);
+		constexpr bool bInLoadIfFindFails = false;
+		// TODO: This is really bad for network travel. Names and paths should be mapped to an uint32. The mapping data is sent by the client once and then reused in the future. 
+		FObjectAndNameAsStringProxyArchive ObjectAndNameProxy(Writer, bInLoadIfFindFails);
+		Object.Serialize(ObjectAndNameProxy);
+
+		FConcertSessionSerializedPayload Payload;
+		Payload.SetTypedPayload(Data);
+		return Payload;
+	}
+
+	void FFullObjectFormat::CombineReplicationEvents(FConcertSessionSerializedPayload& Base, const FConcertSessionSerializedPayload& Newer)
+	{
+		// Full object format always serializes all data. Hence the combination of some Base and Newer is just Newer.
+		// Base and Newer should have the similar size. That means that Base's FConcertSessionSerializedPayload::FConcertByteArray operator= should hopefully reuse the memory allocated for TArray<uint8>
+		Base = Newer; 
+	}
+
+	void FFullObjectFormat::ApplyReplicationEvent(UObject& Object, const FConcertSessionSerializedPayload& Payload)
+	{
+		// This can be quite some data... maybe reading could be done without creating a copy...
+		FFullObjectReplicationData FullObjectFormat;
+		Payload.GetTypedPayload(FullObjectFormat);
+
+		FullObjectFormat::FFulObjectFormatReader WriterArchive(FullObjectFormat.SerializedObjectData.Bytes);
+		constexpr bool bInLoadIfFindFails = false;
+		FObjectAndNameAsStringProxyArchive ObjectAndNameProxy(WriterArchive, bInLoadIfFindFails);
+		Object.Serialize(ObjectAndNameProxy);
+	}
+}
diff --git a/Engine/Plugins/Developer/Concert/ConcertSync/ConcertSyncCore/Source/ConcertSyncCore/Private/Replication/Messages/ConcertReplicationHandshakeMessages.cpp b/Engine/Plugins/Developer/Concert/ConcertSync/ConcertSyncCore/Source/ConcertSyncCore/Private/Replication/Messages/ConcertReplicationHandshakeMessages.cpp
new file mode 100644
--- /dev/null
+++ b/Engine/Plugins/Developer/Concert/ConcertSync/ConcertSyncCore/Source/ConcertSyncCore/Private/Replication/Messages/ConcertReplicationHandshakeMessages.cpp
@@ -0,0 +1,28 @@
+// Copyright Epic Games, Inc. All Rights Reserved.
+
+#include "Replication/Messages/ConcertReplicationHandshakeMessages.h"
+
+namespace UE::ConcertSyncCore::Replication
+{
+	FString LexJoinErrorCode(EJoinReplicationErrorCode ErrorCode)
+	{
+		static_assert(static_cast<int32>(EJoinReplicationErrorCode::Max) == 10, "Update JoinErrorCodeToText when you change entries");
+
+		switch (ErrorCode)
+		{
+		case EJoinReplicationErrorCode::Success: return TEXT("Success");
+		case EJoinReplicationErrorCode::NetworkError: return TEXT("NetworkError");
+		case EJoinReplicationErrorCode::Cancelled: return TEXT("Cancelled");
+		case EJoinReplicationErrorCode::AlreadyInProgress: return TEXT("AlreadyInProgress");
+		case EJoinReplicationErrorCode::NotInAnyConcertSession: return TEXT("NotInAnyConcertSession");
+		case EJoinReplicationErrorCode::InvalidClass: return TEXT("InvalidClass");
+		case EJoinReplicationErrorCode::DuplicateProperty: return TEXT("DuplicateProperty");
+		case EJoinReplicationErrorCode::AlreadyInSession: return TEXT("AlreadyInSession");
+		case EJoinReplicationErrorCode::ConflictingStreamId: return TEXT("ConflictingStreamId");
+		case EJoinReplicationErrorCode::ConflictingAuthority: return TEXT("ConflictingAuthority");
+		case EJoinReplicationErrorCode::FailedToUnpackStream: return TEXT("FailedToUnpackStream");
+		default: 
+			return FString();
+		}
+	}
+}
\ No newline at end of file
diff --git a/Engine/Plugins/Developer/Concert/ConcertSync/ConcertSyncCore/Source/ConcertSyncCore/Private/Replication/Processing/ObjectReplicationCache.cpp b/Engine/Plugins/Developer/Concert/ConcertSync/ConcertSyncCore/Source/ConcertSyncCore/Private/Replication/Processing/ObjectReplicationCache.cpp
new file mode 100644
--- /dev/null
+++ b/Engine/Plugins/Developer/Concert/ConcertSync/ConcertSyncCore/Source/ConcertSyncCore/Private/Replication/Processing/ObjectReplicationCache.cpp
@@ -0,0 +1,126 @@
+// Copyright Epic Games, Inc. All Rights Reserved.
+
+#include "Replication/Processing/ObjectReplicationCache.h"
+
+#include "Replication/ReplicationStreamObjectID.h"
+#include "Replication/Formats/IObjectReplicationFormat.h"
+#include "Replication/Messages/ConcertReplicationEvents.h"
+
+namespace UE::ConcertSyncCore
+{
+	FObjectReplicationCache::FObjectReplicationCache(TSharedRef<IObjectReplicationFormat> ReplicationFormat)
+		: ReplicationFormat(MoveTemp(ReplicationFormat))
+	{}
+
+	int32 FObjectReplicationCache::StoreUntilConsumed(const FGuid& OriginStreamId, const FConcertObjectReplicationEvent& ObjectReplicationEvent)
+	{
+		int32 NumAccepted = 0;
+		
+		const FReplicationStreamObjectID ObjectId{ OriginStreamId, ObjectReplicationEvent.ReplicatedObject };
+		FObjectCache* ObjectCacheBeforeAddition = Cache.Find(ObjectId);
+		if (ObjectCacheBeforeAddition)
+		{
+			// It is important to compare by address and not by TWeakPtr instance because each user has a different TWeakPtr instance (but they all point to the same memory).
+			TSet<FConcertObjectReplicationEvent*> CombineOnceDetection;
+			for (const TPair<TWeakPtr<IReplicationCacheUser>, TWeakPtr<FConcertObjectReplicationEvent>>& InUseDataPair : ObjectCacheBeforeAddition->DataInUse)
+			{
+				const TWeakPtr<FConcertObjectReplicationEvent>& EventData = InUseDataPair.Value;
+				const TSharedPtr<FConcertObjectReplicationEvent> EventDataPin = EventData.Pin();
+				if (EventDataPin && !CombineOnceDetection.Contains(EventDataPin.Get()))
+				{
+					CombineOnceDetection.Add(EventDataPin.Get());
+					ReplicationFormat->CombineReplicationEvents(EventDataPin->SerializedPayload, ObjectReplicationEvent.SerializedPayload);
+				}
+			}
+		}
+
+		FObjectCache* ObjectCacheAfterAddition = ObjectCacheBeforeAddition ? ObjectCacheBeforeAddition : nullptr;
+		TSharedPtr<FConcertObjectReplicationEvent> LazilyCopiedEventPtr;
+		for (const TSharedRef<IReplicationCacheUser>& CacheUser : CacheUsers)
+		{
+			if (ObjectCacheBeforeAddition && ObjectCacheBeforeAddition->DataInUse.Contains(CacheUser))
+			{
+				continue;
+			}
+
+			if (CacheUser->WantsToAcceptObject(ObjectId))
+			{
+				++NumAccepted;
+				
+				// Do the event copy only when somebody wants the data...
+				if (!LazilyCopiedEventPtr.IsValid())
+				{
+					LazilyCopiedEventPtr = MakeShared<FConcertObjectReplicationEvent>(ObjectReplicationEvent);
+				}
+
+				// We'll constructor a new proxy shared ptr which will clean-up Cache when released
+				TWeakPtr<IReplicationCacheUser> WeakUserPtr = CacheUser;
+				TWeakPtr<FObjectReplicationCache> WeakThisPtr = AsWeak();
+				TSharedRef<FConcertObjectReplicationEvent> GuardPtr = MakeShareable(LazilyCopiedEventPtr.Get(), [OriginStreamId, LazilyCopiedEventPtr, WeakUserPtr, WeakThisPtr](auto*)
+				{
+					// Replication cache user can survive the destruction of the cache because it is created externally
+					const TSharedPtr<FObjectReplicationCache> This = WeakThisPtr.Pin();
+					if (!This)
+					{
+						return;
+					}
+
+					// ObjectCache may not be found because cache user was unregistered and then destroyed: UnregisterDataCacheUser removes cache users.
+					const FReplicationStreamObjectID ObjectId{ OriginStreamId, LazilyCopiedEventPtr->ReplicatedObject };
+					FObjectCache* ObjectCache = This->Cache.Find(ObjectId);
+					if (!ObjectCache)
+					{
+						return;
+					}
+					
+					ObjectCache->DataInUse.Remove(WeakUserPtr);
+					if (ObjectCache->DataInUse.IsEmpty())
+					{
+						This->Cache.Remove(ObjectId);
+					}
+					
+					// LazyCopiedEventPtr will decrement counter and possibly be destroyed now
+				});
+				
+				const TWeakPtr<FConcertObjectReplicationEvent> WeakGuardPtr = GuardPtr;
+				CacheUser->OnDataCached(ObjectId, MoveTemp(GuardPtr));
+				// Was it instantly consumed? Should not really happen but it could technically...
+				if (!LIKELY(WeakGuardPtr.IsValid()))
+				{
+					continue;
+				}
+				
+				if (!ObjectCacheAfterAddition)
+				{
+					ObjectCacheAfterAddition = &Cache.Add(ObjectId);
+				}
+				ObjectCacheAfterAddition->DataInUse.Add(WeakUserPtr, WeakGuardPtr);
+			}
+		}
+
+		return NumAccepted;
+	}
+
+	void FObjectReplicationCache::RegisterDataCacheUser(TSharedRef<IReplicationCacheUser> User)
+	{
+		if (!CacheUsers.Contains(User))
+		{
+			CacheUsers.Emplace(MoveTemp(User));
+		}
+	}
+
+	void FObjectReplicationCache::UnregisterDataCacheUser(const TSharedRef<IReplicationCacheUser>& User)
+	{
+		CacheUsers.RemoveSingle(User);
+		
+		for (auto It = Cache.CreateIterator(); It; ++It)
+		{
+			FObjectCache& ObjectCache = It->Value;
+			ObjectCache.DataInUse.Remove(User);
+			if (ObjectCache.DataInUse.IsEmpty())
+			{
+				It.RemoveCurrent();
+			}
+		}
+	}
+}
diff --git a/Engine/Plugins/Developer/Concert/ConcertSync/ConcertSyncCore/Source/ConcertSyncCore/Private/Replication/Processing/ObjectReplicationProcessor.cpp b/Engine/Plugins/Developer/Concert/ConcertSync/ConcertSyncCore/Source/ConcertSyncCore/Private/Replication/Processing/ObjectReplicationProcessor.cpp
new file mode 100644
--- /dev/null
+++ b/Engine/Plugins/Developer/Concert/ConcertSync/ConcertSyncCore/Source/ConcertSyncCore/Private/Replication/Processing/ObjectReplicationProcessor.cpp
@@ -0,0 +1,22 @@
+// Copyright Epic Games, Inc. All Rights Reserved.
+
+#include "Replication/Processing/ObjectReplicationProcessor.h"
+
+#include "ConcertMessageData.h"
+#include "Replication/Processing/IReplicationDataSource.h"
+
+namespace UE::ConcertSyncCore
+{
+	FObjectReplicationProcessor::FObjectReplicationProcessor(TSharedRef<IReplicationDataSource> DataSource)
+		: DataSource(MoveTemp(DataSource))
+	{}
+
+	void FObjectReplicationProcessor::ProcessObjects(float TimeBudget)
+	{
+		// TODO: Respect time budget and prioritize objects
+		DataSource->ForEachPendingObject([this](const FReplicationStreamObjectID& ObjectInfo)
+		{
+			ProcessObject({ ObjectInfo });
+		});
+	}
+}
diff --git a/Engine/Plugins/Developer/Concert/ConcertSync/ConcertSyncCore/Source/ConcertSyncCore/Private/Replication/Processing/ObjectReplicationReceiver.cpp b/Engine/Plugins/Developer/Concert/ConcertSync/ConcertSyncCore/Source/ConcertSyncCore/Private/Replication/Processing/ObjectReplicationReceiver.cpp
new file mode 100644
--- /dev/null
+++ b/Engine/Plugins/Developer/Concert/ConcertSync/ConcertSyncCore/Source/ConcertSyncCore/Private/Replication/Processing/ObjectReplicationReceiver.cpp
@@ -0,0 +1,49 @@
+// Copyright Epic Games, Inc. All Rights Reserved.
+
+#include "Replication/Processing/ObjectReplicationReceiver.h"
+
+#include "ConcertLogGlobal.h"
+#include "IConcertSession.h"
+#include "Replication/Messages/ConcertReplicationEvents.h"
+#include "Replication/Processing/ObjectReplicationCache.h"
+
+namespace UE::ConcertSyncCore
+{
+	FObjectReplicationReceiver::FObjectReplicationReceiver(TSharedRef<IConcertSession> Session, TSharedRef<FObjectReplicationCache> ReplicationCache)
+		: Session(MoveTemp(Session))
+		, ReplicationCache(MoveTemp(ReplicationCache))
+	{
+		Session->RegisterCustomEventHandler<FConcertBatchReplicationEvent>(this, &FObjectReplicationReceiver::HandleBatchReplicationEvent);
+	}
+
+	FObjectReplicationReceiver::~FObjectReplicationReceiver()
+	{
+		Session->UnregisterCustomEventHandler<FConcertBatchReplicationEvent>(this);
+	}
+
+	void FObjectReplicationReceiver::HandleBatchReplicationEvent(const FConcertSessionContext& SessionContext, const FConcertBatchReplicationEvent& Event)
+	{
+		int32 NumObjects = 0;
+		int32 NumCacheUsages = 0;
+		int32 NumOfAcceptedObjects = 0;
+		
+		for (const FConcertStreamReplicationEvent& StreamEvent : Event.Streams)
+		{
+			NumObjects += StreamEvent.ReplicatedObjects.Num();
+			for (const FConcertObjectReplicationEvent& ObjectEvent : StreamEvent.ReplicatedObjects)
+			{
+				const int32 NumAccepted = ReplicationCache->StoreUntilConsumed(StreamEvent.StreamId, ObjectEvent);
+				NumCacheUsages += NumAccepted;
+				NumOfAcceptedObjects += NumAccepted == 0 ? 0 : 1;
+			}
+		}
+		
+		UE_LOG(LogConcert, Verbose, TEXT("Received %d streams with %d objects from endpoint %s. Cached %d objects with a total of %d cache usages."),
+			Event.Streams.Num(),
+			NumObjects,
+			*SessionContext.SourceEndpointId.ToString(),
+			NumOfAcceptedObjects,
+			NumCacheUsages
+			);
+	}
+}
diff --git a/Engine/Plugins/Developer/Concert/ConcertSync/ConcertSyncCore/Source/ConcertSyncCore/Private/Replication/Processing/ObjectReplicationSender.cpp b/Engine/Plugins/Developer/Concert/ConcertSync/ConcertSyncCore/Source/ConcertSyncCore/Private/Replication/Processing/ObjectReplicationSender.cpp
new file mode 100644
--- /dev/null
+++ b/Engine/Plugins/Developer/Concert/ConcertSync/ConcertSyncCore/Source/ConcertSyncCore/Private/Replication/Processing/ObjectReplicationSender.cpp
@@ -0,0 +1,64 @@
+// Copyright Epic Games, Inc. All Rights Reserved.
+
+#include "Replication/Processing/ObjectReplicationSender.h"
+
+#include "ConcertLogGlobal.h"
+#include "IConcertSession.h"
+#include "Algo/Accumulate.h"
+#include "Replication/Processing/IReplicationDataSource.h"
+
+namespace UE::ConcertSyncCore
+{
+	FObjectReplicationSender::FObjectReplicationSender(
+		const FGuid& TargetEndpointId,
+		TSharedRef<IConcertSession> Session,
+		TSharedRef<IReplicationDataSource> DataSource
+		)
+		: FObjectReplicationProcessor(MoveTemp(DataSource))
+		, TargetEndpointId(TargetEndpointId)
+		, Session(MoveTemp(Session))
+	{}
+
+	void FObjectReplicationSender::ProcessObjects(float TimeBudget)
+	{
+		FObjectReplicationProcessor::ProcessObjects(TimeBudget);
+
+		if (!EventToSend.Streams.IsEmpty())
+		{
+			const int32 NumObjects = Algo::TransformAccumulate(EventToSend.Streams, [](const FConcertStreamReplicationEvent& Event){ return Event.ReplicatedObjects.Num(); }, 0);
+			UE_LOG(LogConcert, Verbose, TEXT("Sending %d streams with %d objects to %s"),
+				EventToSend.Streams.Num(),
+				NumObjects,
+				*TargetEndpointId.ToString()
+				);
+			
+			Session->SendCustomEvent(EventToSend, TargetEndpointId,
+				// Replication is always unreliable - if it fails to deliver we'll send updated data soon again
+				// TODO: In regular intervals send CRC values to detect that a change is missing
+				EConcertMessageFlags::None
+			);
+			EventToSend.Streams.Empty(
+				// It's not unreasonable to expect the next pass to have a similar number of objects so keep it around to avoid re-allocating all the time
+				EventToSend.Streams.Num()
+				);
+		}
+	}
+
+	void FObjectReplicationSender::ProcessObject(const FObjectProcessArgs& Args)
+	{
+		const FSoftObjectPath& ReplicatedObject = Args.ObjectInfo.Object;
+		auto CaptureData = [this, &Args, &ReplicatedObject]<typename TPayloadRefType>(TPayloadRefType&& Payload)
+		{
+			const int32 PreexistingIndex = EventToSend.Streams.IndexOfByPredicate([&Args](const FConcertStreamReplicationEvent& StreamData){ return StreamData.StreamId == Args.ObjectInfo.StreamId; });
+			FConcertStreamReplicationEvent& StreamData = EventToSend.Streams.IsValidIndex(PreexistingIndex)
+				? EventToSend.Streams[PreexistingIndex]
+				: EventToSend.Streams[EventToSend.Streams.Emplace(Args.ObjectInfo.StreamId)];
+			StreamData.ReplicatedObjects.Add({
+				ReplicatedObject,
+				// Take advantage of move semantics if it is possible - this depends on how our data source internally obtains its payloads
+				Forward<TPayloadRefType>(Payload)
+			});
+		};
+		GetDataSource().ExtractReplicationDataForObject(Args.ObjectInfo, CaptureData, CaptureData);
+	}
+}
diff --git a/Engine/Plugins/Developer/Concert/ConcertSync/ConcertSyncCore/Source/ConcertSyncCore/Private/Replication/Processing/ReplicationDataQueuer.cpp b/Engine/Plugins/Developer/Concert/ConcertSync/ConcertSyncCore/Source/ConcertSyncCore/Private/Replication/Processing/ReplicationDataQueuer.cpp
new file mode 100644
--- /dev/null
+++ b/Engine/Plugins/Developer/Concert/ConcertSync/ConcertSyncCore/Source/ConcertSyncCore/Private/Replication/Processing/ReplicationDataQueuer.cpp
@@ -0,0 +1,61 @@
+// Copyright Epic Games, Inc. All Rights Reserved.
+
+#include "Replication/Processing/ReplicationDataQueuer.h"
+
+#include "Replication/Messages/ConcertReplicationEvents.h"
+
+namespace UE::ConcertSyncCore
+{
+	FReplicationDataQueuer::~FReplicationDataQueuer()
+	{
+		if (ensure(ReplicationCache))
+		{
+			ReplicationCache->UnregisterDataCacheUser(AsShared());
+		}
+	}
+
+	void FReplicationDataQueuer::ForEachPendingObject(TFunctionRef<void(const FReplicationStreamObjectID&)> ProcessItemFunc) const
+	{
+		for (auto It = PendingEvents.CreateConstIterator(); It; ++It)
+		{
+			ProcessItemFunc(It->Key);
+		}
+	}
+
+	int32 FReplicationDataQueuer::NumObjects() const
+	{
+		return PendingEvents.Num();
+	}
+
+	bool FReplicationDataQueuer::ExtractReplicationDataForObject(
+		const FReplicationStreamObjectID& Object,
+		TFunctionRef<void(const FConcertSessionSerializedPayload& Payload)> ProcessCopyable,
+		TFunctionRef<void(FConcertSessionSerializedPayload&& Payload)> ProcessMoveable
+		)
+	{
+		TSharedPtr<const FConcertObjectReplicationEvent> Event;
+		const bool bSuccess = PendingEvents.RemoveAndCopyValue(Object, Event);
+		if (!ensureMsgf(bSuccess, TEXT("ExtractReplicationDataForObject for an item that was not returned by ForEachPendingObject")))
+		{
+			return false;
+		}
+
+		// Event may be shared by other FReplicationDataQueuers since it originates from the replication cache, so sadly no move.
+		ProcessCopyable(Event->SerializedPayload);
+		return true;
+	}
+
+	void FReplicationDataQueuer::OnDataCached(const FReplicationStreamObjectID& Object, TSharedRef<const FConcertObjectReplicationEvent> Data)
+	{
+		PendingEvents.Add(Object, MoveTemp(Data));
+	}
+	
+	void FReplicationDataQueuer::BindToCache(TSharedRef<FObjectReplicationCache> InReplicationCache)
+	{
+		if (ensure(!ReplicationCache))
+		{
+			ReplicationCache = MoveTemp(InReplicationCache);
+			ReplicationCache->RegisterDataCacheUser(AsShared());
+		}
+	}
+}
diff --git a/Engine/Plugins/Developer/Concert/ConcertSync/ConcertSyncCore/Source/ConcertSyncCore/Private/Replication/PropertyChainUtils.cpp b/Engine/Plugins/Developer/Concert/ConcertSync/ConcertSyncCore/Source/ConcertSyncCore/Private/Replication/PropertyChainUtils.cpp
--- a/Engine/Plugins/Developer/Concert/ConcertSync/ConcertSyncCore/Source/ConcertSyncCore/Private/Replication/PropertyChainUtils.cpp
+++ b/Engine/Plugins/Developer/Concert/ConcertSync/ConcertSyncCore/Source/ConcertSyncCore/Private/Replication/PropertyChainUtils.cpp
@@ -67,7 +67,7 @@ namespace UE::ConcertSyncCore::PropertyChain
 			{
 				// Handle FConcertPropertyChain::InternalContainerPropertyValueName case
 				if (IsPrimitiveProperty(Inner)
-					|| (Inner.IsA(FStructProperty::StaticClass()) && IsNativeStructProperty(Inner)))
+					|| IsNativeStructProperty(Inner))
 				{
 					Chain.PushProperty(&Property, Property.IsEditorOnlyProperty());
 					ON_SCOPE_EXIT{ Chain.PopProperty(&Property, Property.IsEditorOnlyProperty()); };
@@ -181,6 +181,12 @@ namespace UE::ConcertSyncCore::PropertyChain
 			&& (ParentProperty->IsA(FArrayProperty::StaticClass()) || ParentProperty->IsA(FSetProperty::StaticClass()) || ParentProperty->IsA(FMapProperty::StaticClass()));
 		return bIsInnerContainerProperty;
 	}
+	
+	bool IsPropertyEligibleForMarkingAsInternal(const FProperty& Property)
+	{
+		return IsInnerContainerProperty(Property)
+			&& (IsPrimitiveProperty(Property) || IsNativeStructProperty(Property));
+	}
 
 	bool IsPrimitiveProperty(const FProperty& Property)
 	{
diff --git a/Engine/Plugins/Developer/Concert/ConcertSync/ConcertSyncCore/Source/ConcertSyncCore/Private/Replication/ReplicationPropertyFilter.cpp b/Engine/Plugins/Developer/Concert/ConcertSync/ConcertSyncCore/Source/ConcertSyncCore/Private/Replication/ReplicationPropertyFilter.cpp
new file mode 100644
--- /dev/null
+++ b/Engine/Plugins/Developer/Concert/ConcertSync/ConcertSyncCore/Source/ConcertSyncCore/Private/Replication/ReplicationPropertyFilter.cpp
@@ -0,0 +1,57 @@
+// Copyright Epic Games, Inc. All Rights Reserved.
+
+#include "Replication/ReplicationPropertyFilter.h"
+
+#include "Replication/PropertyChainUtils.h"
+#include "Replication/Data/ConcertPropertySelection.h"
+#include "Serialization/ArchiveSerializedPropertyChain.h"
+
+namespace UE::ConcertSyncCore
+{
+	FReplicationPropertyFilter::FReplicationPropertyFilter(const FConcertPropertySelection& PropertySelection)
+		: PropertySelection(PropertySelection)
+	{
+		for (int32 i = 0; i < PropertySelection.ReplicatedProperties.Num(); ++i)
+		{
+			const FConcertPropertyChain& Chain = PropertySelection.ReplicatedProperties[i];
+			LeafToChain.FindOrAdd(Chain.GetLeafProperty()).Add(i);
+		}
+	}
+
+	bool FReplicationPropertyFilter::ShouldSerializeProperty(const FArchiveSerializedPropertyChain* Chain, const FProperty& Property) const
+	{
+		// This would mean Property is 1. in a container and 1.1 primitive or 1.2 a native serialized struct 
+		if (PropertyChain::IsPropertyEligibleForMarkingAsInternal(Property))
+		{
+			if (ensureMsgf(Chain && Chain->GetNumProperties() >= 1, TEXT("Assumption broken that Property is in a container (array, set, map). Check IsPropertyEligibleForMarkingAsInternal implementation!")))
+			{
+				return false;
+			}
+			
+			// We should only have gotten here because a previous call to IsPropertyInSelection(&CopiedChain, *OwnerOfProperty) returned true already (or the property would not have been pushed into the chain).
+			return true;
+		}
+		
+		const TArray<int32>* IndicesToSearch = LeafToChain.Find(Property.GetFName());
+		if (!IndicesToSearch)
+		{
+			/* No chain ends with this property so it is not in the selection
+			 * Note: The property selection must also contain every parent property.
+			 * Example: If MyStruct.Float is in the selection, so must be MyStruct.
+			 * Otherwise this will return false for properties in the middle of the chain (and it should return true).
+			 */
+			return false;
+		}
+
+		for (const int32 IndexToSearch : *IndicesToSearch)
+		{
+			const FConcertPropertyChain& ConcertChain = PropertySelection.ReplicatedProperties[IndexToSearch];
+			if (ConcertChain.MatchesExactly(Chain, Property))
+			{
+				return true;
+			}
+		}
+		
+		return false;
+	}
+}
diff --git a/Engine/Plugins/Developer/Concert/ConcertSync/ConcertSyncCore/Source/ConcertSyncCore/Private/Replication/ReplicationStreamObjectID.cpp b/Engine/Plugins/Developer/Concert/ConcertSync/ConcertSyncCore/Source/ConcertSyncCore/Private/Replication/ReplicationStreamObjectID.cpp
new file mode 100644
--- /dev/null
+++ b/Engine/Plugins/Developer/Concert/ConcertSync/ConcertSyncCore/Source/ConcertSyncCore/Private/Replication/ReplicationStreamObjectID.cpp
@@ -0,0 +1,11 @@
+// Copyright Epic Games, Inc. All Rights Reserved.
+
+#include "Replication/ReplicationStreamObjectID.h"
+
+namespace UE::ConcertSyncCore
+{
+	uint32 GetTypeHash(const FReplicationStreamObjectID& StreamObject)
+	{
+		return HashCombineFast(GetTypeHash(StreamObject.Object), GetTypeHash(StreamObject.StreamId));
+	}
+}
\ No newline at end of file
diff --git a/Engine/Plugins/Developer/Concert/ConcertSync/ConcertSyncCore/Source/ConcertSyncCore/Public/ConcertSyncSessionFlags.h b/Engine/Plugins/Developer/Concert/ConcertSync/ConcertSyncCore/Source/ConcertSyncCore/Public/ConcertSyncSessionFlags.h
--- a/Engine/Plugins/Developer/Concert/ConcertSync/ConcertSyncCore/Source/ConcertSyncCore/Public/ConcertSyncSessionFlags.h
+++ b/Engine/Plugins/Developer/Concert/ConcertSync/ConcertSyncCore/Source/ConcertSyncCore/Public/ConcertSyncSessionFlags.h
@@ -29,28 +29,30 @@ enum class EConcertSyncSessionFlags : uint32
 	EnablePresence = 1<<6,
 	/** Enable usage of file sharing to share large blobs (like package data) rather than embedding that data in message/event.*/
 	EnableFileSharing = 1<<7,
+	/** Enable high-performance property replication. Clients will send data to server and the server replicates it back to other clients interested in that data. */
+	EnableReplication = 1<<8,
 
 	/**
 	 * Flags controlling behavior of enabled features
 	 */
 	
 	/** Discard transactions when a package is unloaded, eg) from a content hot-reload or closing a world editor (requires EnableTransactions) */
-	ShouldDiscardTransactionsOnPackageUnload = 1<<8,
+	ShouldDiscardTransactionsOnPackageUnload = 1<<9,
 	/** Send snapshots of multi-frame transactions (like drags) to other clients (requires EnableTransactions) */
-	ShouldSendTransactionSnapshots = 1<<9,
+	ShouldSendTransactionSnapshots = 1<<10,
 	/** Send the pristine state of a package (as revision 0) before saving it for the first time (requires EnablePackages) */
-	ShouldSendPackagePristineState = 1<<10,
+	ShouldSendPackagePristineState = 1<<11,
 	/** Send package auto-saves, in addition to user-saves (requires EnablePackages) */
-	ShouldSendPackageAutoSaves = 1<<11,
+	ShouldSendPackageAutoSaves = 1<<12,
 	/** Use a sandbox to hold package writes during a session (requires EnablePackages) */
-	ShouldUsePackageSandbox = 1<<12,
+	ShouldUsePackageSandbox = 1<<13,
 
 	/**
 	 * Defaults for different scenarios
 	 */
 	
 	/** Default settings for a multi-user session */
-	Default_MultiUserSession = EnableLiveSync | EnableConnectionHistory | EnableLocking | EnableTransactions | EnablePackages | EnableSequencer | EnablePresence | ShouldSendTransactionSnapshots | ShouldUsePackageSandbox,
+	Default_MultiUserSession = EnableLiveSync | EnableConnectionHistory | EnableLocking | EnableTransactions | EnablePackages | EnableSequencer | EnablePresence | ShouldSendTransactionSnapshots | ShouldUsePackageSandbox | EnableReplication,
 	/** Default settings for a disaster recovery session */
 	Default_DisasterRecoverySession = EnableTransactions | EnablePackages | EnableFileSharing | ShouldDiscardTransactionsOnPackageUnload | ShouldSendPackagePristineState | ShouldSendPackageAutoSaves,
 };
diff --git a/Engine/Plugins/Developer/Concert/ConcertSync/ConcertSyncCore/Source/ConcertSyncCore/Public/Replication/Data/ConcertObjectReplicationSettings.h b/Engine/Plugins/Developer/Concert/ConcertSync/ConcertSyncCore/Source/ConcertSyncCore/Public/Replication/Data/ConcertObjectReplicationSettings.h
new file mode 100644
--- /dev/null
+++ b/Engine/Plugins/Developer/Concert/ConcertSync/ConcertSyncCore/Source/ConcertSyncCore/Public/Replication/Data/ConcertObjectReplicationSettings.h
@@ -0,0 +1,17 @@
+// Copyright Epic Games, Inc. All Rights Reserved.
+
+#pragma once
+
+#include "ConcertObjectReplicationSettings.generated.h"
+
+/** Describes how a specific object is supposed to be replicated. */
+USTRUCT()
+struct FConcertObjectReplicationSettings
+{
+	GENERATED_BODY()
+
+	/** How often per second this object should be updated. */
+	UPROPERTY(EditAnywhere, Category = "Concert")
+	float UpdateFrequency = 30.f;
+};
+
diff --git a/Engine/Plugins/Developer/Concert/ConcertSync/ConcertSyncCore/Source/ConcertSyncCore/Public/Replication/Data/ConcertPropertySelection.h b/Engine/Plugins/Developer/Concert/ConcertSync/ConcertSyncCore/Source/ConcertSyncCore/Public/Replication/Data/ConcertPropertySelection.h
--- a/Engine/Plugins/Developer/Concert/ConcertSync/ConcertSyncCore/Source/ConcertSyncCore/Public/Replication/Data/ConcertPropertySelection.h
+++ b/Engine/Plugins/Developer/Concert/ConcertSync/ConcertSyncCore/Source/ConcertSyncCore/Public/Replication/Data/ConcertPropertySelection.h
@@ -7,13 +7,28 @@
 class UStruct;
 struct FArchiveSerializedPropertyChain;
 
-/** Describes the path to a FProperty replicated by Concert.  */
+/**
+ * Describes the path to a FProperty replicated by Concert.
+ * 
+ * See Concert.Replication.Data.ForEachReplicatableConcertProperty for path examples (just CTRL+SHIFT+F or go to ConcertSyncTest/Replication/ConcertPropertyTests.cpp).
+ */
 USTRUCT()
 struct CONCERTSYNCCORE_API FConcertPropertyChain
 {
 	GENERATED_BODY()
 
-	/** The name of inner properties of containers (array, set, map) which are either primitive or native structs. */
+	/**
+	 * The name of inner properties of containers (array, set, map), which are either primitive or native serialized structs.
+	 * 
+	 * Note that inner properties are only included if they are at the end of the path (and only for primitive or native serialized structs!), e.g.
+	 * - ArrayOfStructs.ArrayOfFloats.Value
+	 *	This means there is an array property called ArrayOfStructs containing structs. The contained struct has an array of floats property
+	 *	called ArrayOfFloats. Value corresponds to FArrayProperty::Inner and is called InternalContainerPropertyValueName.
+	 *	
+	 * - but not ArrayOfStructs.Value.ArrayOfFloats.Value.
+	 *	This path would imply a different situation in which ArrayOfStructs contains a struct with a struct property called Value;
+	 *	Value's owning struct property would contain an array of floats called ArrayOfFloats.
+	 */
 	static const FName InternalContainerPropertyValueName;
 
 	/** Constructs a FConcertPropertyChain from a path if it is valid. If you need to create many paths in one go, use PropertyUtils::BulkConstructConcertChainsFromPaths instead. */
@@ -38,6 +53,9 @@ struct CONCERTSYNCCORE_API FConcertPropertyChain
 	/** @return Whether the leaf property is a direct child of the given property chain. */
 	bool IsDirectChildOf(const FConcertPropertyChain& ParentToCheck) const;
 
+	/** @return Whether OptionalChain and LeafProperty correspond to this path. */
+	bool MatchesExactly(const FArchiveSerializedPropertyChain* OptionalChain, const FProperty& LeafProperty) const;
+
 	const TArray<FName>& GetPathToProperty() const { return PathToProperty; }
 
 	enum class EToStringMethod
@@ -80,6 +98,8 @@ private:
 	 * Path from root of UObject to leaf property. Includes the leaf property.
 	 * This property is kept private to force the use of the exposed constructors.
 	 *
+	 * See Concert.Replication.Data.ForEachReplicatableConcertProperty for path examples (just CTRL+SHIFT+F or go to ConcertSyncTest/Replication/ConcertPropertyTests.cpp).
+	 *
 	 * Suppose:
 	 * class AFooActor
 	 * {
@@ -125,7 +145,7 @@ private:
 	 *	  the key inner property nor any of its subproperties. Example: If key is FSoftObjectPath, we do not list FSoftObjectPath::AssetPath, etc.
 	 *	  - Value:
 	 *	    - If the inner value property is is ustruct, then it is not listed. The child properties of the inner property are listed.
-	 *	    - If the inner value property is not a ustruct (meaning a primitive float, int, etc.), then the child property is listed and called "Value".
+	 *	    - If the inner value property is not a ustruct (meaning a primitive float, int, etc.), then the child property is listed and called "Value" (see InternalContainerPropertyValueName).
 	 *	      This is needed to differentiate between the cases "serializes only the keys" and "serializes the keys and the values".
 	 *	  - Final word about FConcertPropertySelection:
 	 *		- { "MapOfStructs", "Foo" } means all keys and the Foo property is serialized (nothing is said about Bar - it is replicated if it is also in the property selection).
diff --git a/Engine/Plugins/Developer/Concert/ConcertSync/ConcertSyncCore/Source/ConcertSyncCore/Public/Replication/Data/ObjectReplicationMap.h b/Engine/Plugins/Developer/Concert/ConcertSync/ConcertSyncCore/Source/ConcertSyncCore/Public/Replication/Data/ObjectReplicationMap.h
--- a/Engine/Plugins/Developer/Concert/ConcertSync/ConcertSyncCore/Source/ConcertSyncCore/Public/Replication/Data/ObjectReplicationMap.h
+++ b/Engine/Plugins/Developer/Concert/ConcertSync/ConcertSyncCore/Source/ConcertSyncCore/Public/Replication/Data/ObjectReplicationMap.h
@@ -14,7 +14,7 @@ struct FReplicatedObjectInfo
 	FSoftClassPath ClassPath;
 	
 	UPROPERTY()
-	FConcertPropertySelection ReplicatedProperties;
+	FConcertPropertySelection PropertySelection;
 };
 
 /** Maps objects to their replicated properties. */
diff --git a/Engine/Plugins/Developer/Concert/ConcertSync/ConcertSyncCore/Source/ConcertSyncCore/Public/Replication/Data/ReplicationStreamDescription.h b/Engine/Plugins/Developer/Concert/ConcertSync/ConcertSyncCore/Source/ConcertSyncCore/Public/Replication/Data/ReplicationStreamDescription.h
--- a/Engine/Plugins/Developer/Concert/ConcertSync/ConcertSyncCore/Source/ConcertSyncCore/Public/Replication/Data/ReplicationStreamDescription.h
+++ b/Engine/Plugins/Developer/Concert/ConcertSync/ConcertSyncCore/Source/ConcertSyncCore/Public/Replication/Data/ReplicationStreamDescription.h
@@ -21,6 +21,8 @@ struct FSharedReplicationStreamDescription
 	/** Identifies the data that this stream will send. */
 	UPROPERTY()
 	FObjectReplicationMap ReplicationMap;
+
+	// TODO: Add override settings to objects, such as override update frequency etc. Do it via TMap<FSoftObjectPath, FObjectReplicationSettings>.
 };
 
 /**
diff --git a/Engine/Plugins/Developer/Concert/ConcertSync/ConcertSyncCore/Source/ConcertSyncCore/Public/Replication/Formats/FullObjectFormat.h b/Engine/Plugins/Developer/Concert/ConcertSync/ConcertSyncCore/Source/ConcertSyncCore/Public/Replication/Formats/FullObjectFormat.h
new file mode 100644
--- /dev/null
+++ b/Engine/Plugins/Developer/Concert/ConcertSync/ConcertSyncCore/Source/ConcertSyncCore/Public/Replication/Formats/FullObjectFormat.h
@@ -0,0 +1,31 @@
+// Copyright Epic Games, Inc. All Rights Reserved.
+
+#pragma once
+
+#include "IObjectReplicationFormat.h"
+
+namespace UE::ConcertSyncCore
+{
+	/**
+	 * Full object format is a blob of serialized UObject data. 
+	 *
+	 * Since there no additional property information in this format, the server cannot merge data. Every update the
+	 * entire blob must be sent. This format is inefficient. The only upside it that it is straight-forward to understand.
+	 * 
+	 * This format is used for the initial Concert Replication prototype.
+	 * TODO: After the prototype phase, more efficient formats should be developed.
+	 *
+	 * @see FFullObjectReplicationData
+	 */
+	class CONCERTSYNCCORE_API FFullObjectFormat : public IObjectReplicationFormat
+	{
+	public:
+
+		//~ Begin FFullObjectFormat Interface
+		virtual TOptional<FConcertSessionSerializedPayload> CreateReplicationEvent(UObject& Object, FAllowPropertyFunc IsPropertyAllowedFunc) override;
+		virtual void ClearInternalCache(TArrayView<UObject> ObjectsToClear) override { /* There is no cache: Full object format has no smart building relying on past events - it always serializes everything */ }
+		virtual void CombineReplicationEvents(FConcertSessionSerializedPayload& Base, const FConcertSessionSerializedPayload& Newer) override;
+		virtual void ApplyReplicationEvent(UObject& Object, const FConcertSessionSerializedPayload& Payload) override;
+		//~ End FFullObjectFormat Interface
+	};
+}
diff --git a/Engine/Plugins/Developer/Concert/ConcertSync/ConcertSyncCore/Source/ConcertSyncCore/Public/Replication/Formats/FullObjectReplicationData.h b/Engine/Plugins/Developer/Concert/ConcertSync/ConcertSyncCore/Source/ConcertSyncCore/Public/Replication/Formats/FullObjectReplicationData.h
new file mode 100644
--- /dev/null
+++ b/Engine/Plugins/Developer/Concert/ConcertSync/ConcertSyncCore/Source/ConcertSyncCore/Public/Replication/Formats/FullObjectReplicationData.h
@@ -0,0 +1,18 @@
+// Copyright Epic Games, Inc. All Rights Reserved.
+
+#pragma once
+
+#include "FullObjectReplicationData.generated.h"
+
+/**
+ * Full object format is a blob of serialized UObject data.
+ * @see FFullObjectFormat
+ */
+USTRUCT()
+struct FFullObjectReplicationData
+{
+	GENERATED_BODY()
+
+	UPROPERTY()
+	FConcertByteArray SerializedObjectData;
+};
\ No newline at end of file
diff --git a/Engine/Plugins/Developer/Concert/ConcertSync/ConcertSyncCore/Source/ConcertSyncCore/Public/Replication/Formats/IObjectReplicationFormat.h b/Engine/Plugins/Developer/Concert/ConcertSync/ConcertSyncCore/Source/ConcertSyncCore/Public/Replication/Formats/IObjectReplicationFormat.h
new file mode 100644
--- /dev/null
+++ b/Engine/Plugins/Developer/Concert/ConcertSync/ConcertSyncCore/Source/ConcertSyncCore/Public/Replication/Formats/IObjectReplicationFormat.h
@@ -0,0 +1,67 @@
+// Copyright Epic Games, Inc. All Rights Reserved.
+
+#pragma once
+
+#include "Containers/ArrayView.h"
+#include "Templates/Function.h"
+#include "Misc/Optional.h"
+
+class FProperty;
+class UObject;
+struct FArchiveSerializedPropertyChain;
+struct FConcertSessionSerializedPayload;
+
+namespace UE::ConcertSyncCore
+{
+	/**
+	 * Abstracts the concept of a replication format for data passed into FConcertReplicationEvent::SerializedPayload.
+	 *
+	 * Example implementations:
+	 * - full-object format that serializes all replicated properties and never does any combining (inefficient).
+	 *   TODO: Full-object is for the MVP prototype - should we remove it later?
+	 * - partial-object format that keeps track of past properties and only sends changed replicated properties.
+	 */
+	class CONCERTSYNCCORE_API IObjectReplicationFormat
+	{
+	public:
+
+		/** Chain can be nullptr or non-null and empty - both imply the Property is a root property. */
+		using FAllowPropertyFunc = TFunctionRef<bool(const FArchiveSerializedPropertyChain* Chain, const FProperty& Property)>;
+
+		/**
+		 * Produces a new payload from Object.
+		 * TODO: Probably the signature of this function will be extended so you can build from a passed-in FConcertSessionSerializedPayload. Maybe for server sending to clients.
+		 * 
+		 * Implementations may create a diff from a previous version of the Object. This function may update the internal cache.
+		 * If there is not data to apply since the last call, then this returns an empty optional.
+		 *
+		 * @param Object The object to serialize
+		 * @param IsPropertyAllowedFunc Decides whether a given property is generally supposed to be replicated.
+		 *		  Note: IObjectReplicationFormat does diffing itself - only return whether you ever want to replicate the property!
+		 *
+		 * @return	Data to be placed into FConcertReplicationEvent::SerializedPayload. This can be a diff from the last time CreateReplicationEvent was called.
+		 *			If there is no data to apply since the last time, the optional is empty.
+		 */
+		virtual TOptional<FConcertSessionSerializedPayload> CreateReplicationEvent(UObject& Object, FAllowPropertyFunc IsPropertyAllowedFunc) = 0;
+
+		/**
+		 * Clears the internal cache that subclass implementations may rely on for building diffed payload version in CreateReplicationEvent.
+		 * The next call to CreateReplicationEvent will serialize all data.
+		 */
+		virtual void ClearInternalCache(TArrayView<UObject> ObjectsToClear) = 0;
+
+		/**
+		 * Combines Base and Newer to form a single payload.
+		 * Calling ApplyReplicationEvent(CombineReplicationEvents(Base, Other)) should be equivalent to calling ApplyReplicationEvent(Base) followed by ApplyReplicationEvent(Newer).
+		 * 
+		 * @param Base Payload produced by CreateReplicationEvent. Produced before Newer. This will be overwritten to contain the combined result.
+		 * @param Newer Payload produced by CreateReplicationEvent. Produced after Base and on the same UObject.
+		 */
+		virtual void CombineReplicationEvents(FConcertSessionSerializedPayload& Base, const FConcertSessionSerializedPayload& Newer) = 0;
+
+		/** Applies Payload, which was created by CreateReplicationEvent, to Object. */
+		virtual void ApplyReplicationEvent(UObject& Object, const FConcertSessionSerializedPayload& Payload) = 0;
+
+		virtual ~IObjectReplicationFormat() = default;
+	};
+}
\ No newline at end of file
diff --git a/Engine/Plugins/Developer/Concert/ConcertSync/ConcertSyncCore/Source/ConcertSyncCore/Public/Replication/Messages/ConcertReplicationEvents.h b/Engine/Plugins/Developer/Concert/ConcertSync/ConcertSyncCore/Source/ConcertSyncCore/Public/Replication/Messages/ConcertReplicationEvents.h
new file mode 100644
--- /dev/null
+++ b/Engine/Plugins/Developer/Concert/ConcertSync/ConcertSyncCore/Source/ConcertSyncCore/Public/Replication/Messages/ConcertReplicationEvents.h
@@ -0,0 +1,60 @@
+// Copyright Epic Games, Inc. All Rights Reserved.
+
+#pragma once
+
+#include "ConcertMessageData.h"
+#include "ConcertReplicationEvents.generated.h"
+
+/** Contains data to be applied to a replicated object. */
+USTRUCT()
+struct FConcertObjectReplicationEvent
+{
+	GENERATED_BODY()
+
+	/**
+	 * The object to serialize into.
+	 * TODO: To reduce network load, replace this with an int32 that maps into FObjectToIdTable (custom) / UPackageMap (engine)  shared by all clients and the server.
+	 */
+	UPROPERTY()
+	FSoftObjectPath ReplicatedObject;
+	
+	/** Contains another struct as payload, such as FConcertFullObjectReplicationData. */
+	UPROPERTY()
+	FConcertSessionSerializedPayload SerializedPayload;
+};
+
+/** Contains data produced by a specific replication stream. */
+USTRUCT()
+struct FConcertStreamReplicationEvent
+{
+	GENERATED_BODY()
+
+	/** The stream that produced this data. */
+	UPROPERTY()
+	FGuid StreamId;
+
+	/** The objects replicated by this event */
+	UPROPERTY()
+	TArray<FConcertObjectReplicationEvent> ReplicatedObjects;
+
+	FConcertStreamReplicationEvent() = default;
+	explicit FConcertStreamReplicationEvent(const FGuid& StreamId)
+		: StreamId(StreamId)
+	{}
+};
+
+/**
+ * Contains multiple objects to be replicated.
+ * Sent from
+ *  - Authoritative client to server
+ *  - Server to clients. Every event may be customized for each client (clients can choose to ignore object updates).
+ */
+USTRUCT()
+struct FConcertBatchReplicationEvent
+{
+	GENERATED_BODY()
+	
+	/** The objects replicated by this event */
+	UPROPERTY()
+	TArray<FConcertStreamReplicationEvent> Streams;
+};
\ No newline at end of file
diff --git a/Engine/Plugins/Developer/Concert/ConcertSync/ConcertSyncCore/Source/ConcertSyncCore/Public/Replication/Messages/ConcertReplicationHandshakeMessages.h b/Engine/Plugins/Developer/Concert/ConcertSync/ConcertSyncCore/Source/ConcertSyncCore/Public/Replication/Messages/ConcertReplicationHandshakeMessages.h
new file mode 100644
--- /dev/null
+++ b/Engine/Plugins/Developer/Concert/ConcertSync/ConcertSyncCore/Source/ConcertSyncCore/Public/Replication/Messages/ConcertReplicationHandshakeMessages.h
@@ -0,0 +1,100 @@
+// Copyright Epic Games, Inc. All Rights Reserved.
+
+#pragma once
+
+#include "Replication/Data/ReplicationStreamDescription.h"
+#include "ConcertReplicationHandshakeMessages.generated.h"
+
+class FText;
+
+UENUM()
+enum class EJoinReplicationErrorCode : uint8
+{
+	Success = 0,
+	/** Indicates generic error obtaining a response from the server. Responses are default constructed if an error occurs. This is the default value. */
+	NetworkError = 1,
+
+	// Client-side error codes:
+	
+	/** Cancelled while initialization was in progress. Most likely LeaveReplicationSession was called while handshake was in progress. */
+	Cancelled = 2,
+	/** Another request is already in progress. */
+	AlreadyInProgress = 3,
+
+	// Server error codes:
+
+	/** Rare case in which client is no longer in a session but has requested to join a replication session. */
+	NotInAnyConcertSession = 4,
+	
+	/** One of the specified classes was null or otherwise invalid */
+	InvalidClass = 5,
+	
+	/** An object selection contained a property twice. */
+	DuplicateProperty = 6,
+	
+	/**
+	 * Client has already joined a replication session. Send FConcertLeaveReplicationSessionEvent first.
+	 */
+	AlreadyInSession = 7,
+
+	/** You or another client already is using the same stream id. */
+	ConflictingStreamId = 8,
+	
+	/** You or another client already has authority over some of the data. */
+	ConflictingAuthority = 9,
+
+	/** Failed to unpack the stream, likely because one of the attribute classes could not be resolved. */
+	FailedToUnpackStream = 10,
+	
+	MaxPlusOne,
+	/** The maximum possible entry value */
+	Max = MaxPlusOne - 1
+};
+
+namespace UE::ConcertSyncCore::Replication
+{
+	CONCERTSYNCCORE_API FString LexJoinErrorCode(EJoinReplicationErrorCode ErrorCode);
+}
+
+/**
+ * Sent by client to request joining a replication session.
+ * This is sent by clients want to send or receive replicated data, or both.
+ */
+USTRUCT()
+struct FConcertReplication_Join_Request
+{
+	GENERATED_BODY()
+
+	/** The data the client offers to send. */
+	UPROPERTY()
+	TArray<FReplicationStreamDescription_NetPacked> Streams;
+};
+
+/**
+ * Sent by server to accept the previously sent replication session.
+ */
+USTRUCT()
+struct FConcertReplication_Join_Response
+{
+	GENERATED_BODY()
+
+	/**
+	 * Anything other than Success indicates failure.
+	 * 
+	 * Remember that Concert default constructs a response when it fails to obtain a response.
+	 * Hence NetworkError means timeout, etc.
+	 */
+	UPROPERTY()
+	EJoinReplicationErrorCode ErrorCode = EJoinReplicationErrorCode::NetworkError;
+
+	/** More information about ErrorCode to help the user resolve the issue. */
+	UPROPERTY()
+	FString DetailedErrorMessage;
+};
+
+/** Sent by client to notify server that no more replication data should be sent. */
+USTRUCT()
+struct FConcertReplication_LeaveEvent
+{
+	GENERATED_BODY()
+};
\ No newline at end of file
diff --git a/Engine/Plugins/Developer/Concert/ConcertSync/ConcertSyncCore/Source/ConcertSyncCore/Public/Replication/Processing/IReplicationDataSource.h b/Engine/Plugins/Developer/Concert/ConcertSync/ConcertSyncCore/Source/ConcertSyncCore/Public/Replication/Processing/IReplicationDataSource.h
new file mode 100644
--- /dev/null
+++ b/Engine/Plugins/Developer/Concert/ConcertSync/ConcertSyncCore/Source/ConcertSyncCore/Public/Replication/Processing/IReplicationDataSource.h
@@ -0,0 +1,69 @@
+// Copyright Epic Games, Inc. All Rights Reserved.
+
+#pragma once
+
+#include "ConcertMessageData.h"
+#include "HAL/Platform.h"
+#include "Misc/Optional.h"
+#include "Templates/Function.h"
+
+namespace UE::ConcertSyncCore
+{
+	struct FReplicationStreamObjectID;
+	
+	/**
+	 * Responsible for obtaining replication data.
+	 * 
+	 * Implementation examples:
+	 * - serialize the object path (client),
+	 * - cache received object data and return it (client & server)
+	 */
+	class IReplicationDataSource
+	{
+	public:
+		
+		/** Iterates the objects that must be processed for replication (the result of ExtractReplicationDataForObject MAY return something new). */
+		virtual void ForEachPendingObject(TFunctionRef<void(const FReplicationStreamObjectID&)> ProcessItemFunc) const = 0;
+		/** @return The number of objects ForEachPendingObject would iterate. Can be used, e.g. for Reserve()-ing a container. */
+		virtual int32 NumObjects() const = 0;
+
+		/**
+		 * Extracts data for Object. Object must have been previously returned by ForEachPendingObject.
+		 * This "dequeues" Object so it will not be mentioned by ForEachPendingObject again until it is again marked as "dirty".
+		 *
+		 * The call to ProcessCopyable / ProcessMoveable may be skipped if there is no new data to send.
+		 * Either ProcessCopyable or ProcessMoveable will be called, never both, and it will be called at most once.
+		 * @param Object The object for which to obtain data
+		 * @param ProcessCopyable Callback if the event was retrieved and not owned by this IReplicationDataSource (hence not being moveable).
+		 * @param ProcessMoveable Callback if the event was just constructed (and hence can be moved)
+		 * 
+		 * @return Whether successful. False indicates the call was invalid to make (ForEachPendingObject did not return Object). True indicates success, even if ProcessCopyable was not called.		
+		 */
+		virtual bool ExtractReplicationDataForObject(
+			const FReplicationStreamObjectID& Object,
+			TFunctionRef<void(const FConcertSessionSerializedPayload& Payload)> ProcessCopyable,
+			TFunctionRef<void(FConcertSessionSerializedPayload&& Payload)> ProcessMoveable
+			) = 0;
+		/** Util version for callers that only want to read and do not want to store the payload. */
+		bool ExtractReplicationDataForObject(const FReplicationStreamObjectID& Object, TFunctionRef<void(const FConcertSessionSerializedPayload& Payload)> ProcessCopyable)
+		{
+			return ExtractReplicationDataForObject(Object, ProcessCopyable, [&ProcessCopyable](FConcertSessionSerializedPayload&& Payload){ ProcessCopyable(Payload); });
+		}
+
+		/** Util function for add the contained objects to a container (TArray / TSet). Container must have Reserve and Add functions.  */
+		template<typename TContainerType>
+		void AddEachObjectToContainer(TContainerType& OutContainer) const;
+		
+		virtual ~IReplicationDataSource() = default;
+	};
+
+	template <typename TContainerType>
+	void IReplicationDataSource::AddEachObjectToContainer(TContainerType& OutContainer) const
+	{
+		OutContainer.Reserve(OutContainer.Num() + NumObjects());
+		ForEachPendingObject([&OutContainer](const FSoftObjectPath& Item) mutable
+		{
+			OutContainer.Add(Item);
+		});
+	}
+}
diff --git a/Engine/Plugins/Developer/Concert/ConcertSync/ConcertSyncCore/Source/ConcertSyncCore/Public/Replication/Processing/ObjectReplicationCache.h b/Engine/Plugins/Developer/Concert/ConcertSync/ConcertSyncCore/Source/ConcertSyncCore/Public/Replication/Processing/ObjectReplicationCache.h
new file mode 100644
--- /dev/null
+++ b/Engine/Plugins/Developer/Concert/ConcertSync/ConcertSyncCore/Source/ConcertSyncCore/Public/Replication/Processing/ObjectReplicationCache.h
@@ -0,0 +1,96 @@
+// Copyright Epic Games, Inc. All Rights Reserved.
+
+#pragma once
+
+#include "Containers/Array.h"
+#include "Delegates/DelegateCombinations.h"
+#include "Replication/ReplicationStreamObjectID.h"
+#include "Templates/Function.h"
+#include "Templates/SharedPointer.h"
+
+struct FConcertObjectReplicationEvent;
+struct FSoftObjectPath;
+struct FGuid;
+
+namespace UE::ConcertSyncCore
+{
+	class IObjectReplicationFormat;
+
+	/**
+	 * Interface for objects that want to latently use received FConcertObjectReplicationEvent data.
+	 * The data is continuously updated until it is consumed.
+	 */
+	class CONCERTSYNCCORE_API IReplicationCacheUser
+	{
+	public:
+
+		/** @return Whether this user is interested in data from this object. */
+		virtual bool WantsToAcceptObject(const FReplicationStreamObjectID& Object) const = 0;
+
+		/**
+		 * Called when data that is interesting to this user becomes available.
+		 *
+		 * The user can keep hold of Data until it is used, at which point it just let's Data get out of scope.
+		 * If new data is received while this user is referencing Data, Data will be combined to contain any new data.
+		 */
+		virtual void OnDataCached(const FReplicationStreamObjectID& Object, TSharedRef<const FConcertObjectReplicationEvent> Data) = 0;
+
+		virtual ~IReplicationCacheUser() = default;
+	};
+	
+	/**
+	 * This is an intermediate place for received data to live before it is further processed.
+	 *
+	 * IObjectCacheUsers register with the cache and decide which data is to be received.
+	 * When replication data comes in IObjectCacheUsers::WantsToAcceptObject is used to determine whether the object wants the data.
+	 * If the data should be received, IObjectCacheUsers::OnDataCached is called receiving a shared ptr to the data.
+	 * When the data is finally consumed latently, e.g. sent to other endpoints, the cache user resets the shared ptr.
+	 * If new data comes in before a cache user consumes it, the new data and old data are combined (using IObjectReplicationFormat::CombineReplicationEvents). 
+	 *
+	 * This allows multiple systems to reuse replication data. For example, on the server the same data may need to be distributed to different clients but
+	 * the clients send the data at different times.
+	 */
+	class CONCERTSYNCCORE_API FObjectReplicationCache : public TSharedFromThis<FObjectReplicationCache>
+	{
+	public:
+
+		FObjectReplicationCache(TSharedRef<IObjectReplicationFormat> ReplicationFormat);
+
+		/**
+		 * Called when new data is received for an object and shares it with any IObjectCacheUser that is possibly interested in it.
+		 * @param OriginStreamId The stream from which the object was replicated
+		 * @param ObjectReplicationEvent The data that was replicated
+		 * @return The number of cache users that accepted this event. 
+		 */
+		int32 StoreUntilConsumed(const FGuid& OriginStreamId, const FConcertObjectReplicationEvent& ObjectReplicationEvent);
+
+		/** Registers a new user, which will start receiving data for any new data received from now on. */
+		void RegisterDataCacheUser(TSharedRef<IReplicationCacheUser> User);
+		void UnregisterDataCacheUser(const TSharedRef<IReplicationCacheUser>& User);
+
+	private:
+
+		/** Used for combining events to save network bandwidth. */
+		TSharedRef<IObjectReplicationFormat> ReplicationFormat;
+
+		/** Everyone who registered for receiving data. */
+		TArray<TSharedRef<IReplicationCacheUser>> CacheUsers;
+
+		struct FObjectCache
+		{
+			/**
+			 * Past data that is still in use by users.
+			 * 
+			 * StoreUntilConsumed asks all IReplicationCacheUser and if at least one is interested, creates exactly 1 TSharedRef<FConcertObjectReplicationEvent>.
+			 * Every IReplicationCacheUser receives essentially a MakeShareable of the above instance: its own TSharedRef with a custom deleter that simply keeps the event instance.
+			 *
+			 * This mechanism allows detecting whether a IReplicationCacheUser already is using old data which needs to be combined with the new incoming data, or needs a new instance.
+			 * The intention is that as soon as IReplicationCacheUser has finished using a data event, it will receive a
+			 * new instance: we do not want IReplicationCacheUsers to have large histories of events being combined into them.
+			 */
+			TMap<TWeakPtr<IReplicationCacheUser>, TWeakPtr<FConcertObjectReplicationEvent>> DataInUse;
+		};
+		/** Maps every object to the events cached for it. */
+		TMap<FReplicationStreamObjectID, FObjectCache> Cache; 
+	};
+}
\ No newline at end of file
diff --git a/Engine/Plugins/Developer/Concert/ConcertSync/ConcertSyncCore/Source/ConcertSyncCore/Public/Replication/Processing/ObjectReplicationProcessor.h b/Engine/Plugins/Developer/Concert/ConcertSync/ConcertSyncCore/Source/ConcertSyncCore/Public/Replication/Processing/ObjectReplicationProcessor.h
new file mode 100644
--- /dev/null
+++ b/Engine/Plugins/Developer/Concert/ConcertSync/ConcertSyncCore/Source/ConcertSyncCore/Public/Replication/Processing/ObjectReplicationProcessor.h
@@ -0,0 +1,53 @@
+// Copyright Epic Games, Inc. All Rights Reserved.
+
+#pragma once
+
+#include "IReplicationDataSource.h"
+#include "Replication/ReplicationStreamObjectID.h"
+#include "Templates/SharedPointer.h"
+#include "UObject/SoftObjectPath.h"
+
+namespace UE::ConcertSyncCore
+{
+	class IReplicationDataSource;
+	
+	/**
+	 * Responsible for prioritizing a list of objects and processing them.
+	 * 
+	 * Processing is defined by the subclass. Some example implementations:
+	 *  - serialize the object and send (client),
+	 *  - find loaded object and apply replication data (client)
+	 *  - send object data to interested clients (server),
+	 */
+	class CONCERTSYNCCORE_API FObjectReplicationProcessor
+	{
+	public:
+
+		explicit FObjectReplicationProcessor(TSharedRef<IReplicationDataSource> DataSource);
+		virtual ~FObjectReplicationProcessor() = default;
+
+		/**
+		 * Processes all changed objects under the given time budget.
+		 * The time budget may be exceeded but we'll try not to and to stay as close to the budget as possible.
+		 */
+		virtual void ProcessObjects(float TimeBudget);
+		
+	protected:
+
+		struct FObjectProcessArgs
+		{
+			/** Info about the object to process */
+			FReplicationStreamObjectID ObjectInfo;
+		};
+
+		FORCEINLINE IReplicationDataSource& GetDataSource() const { return DataSource.Get(); }
+		
+		/** Processes the object. */
+		virtual void ProcessObject(const FObjectProcessArgs& Args) = 0;
+
+	private:
+		
+		/** Abstracts where replication data comes from: could be generated (clients) or received (server or client) */
+		TSharedRef<IReplicationDataSource> DataSource;
+	};
+}
diff --git a/Engine/Plugins/Developer/Concert/ConcertSync/ConcertSyncCore/Source/ConcertSyncCore/Public/Replication/Processing/ObjectReplicationReceiver.h b/Engine/Plugins/Developer/Concert/ConcertSync/ConcertSyncCore/Source/ConcertSyncCore/Public/Replication/Processing/ObjectReplicationReceiver.h
new file mode 100644
--- /dev/null
+++ b/Engine/Plugins/Developer/Concert/ConcertSync/ConcertSyncCore/Source/ConcertSyncCore/Public/Replication/Processing/ObjectReplicationReceiver.h
@@ -0,0 +1,33 @@
+// Copyright Epic Games, Inc. All Rights Reserved.
+
+#pragma once
+
+#include "Templates/SharedPointer.h"
+
+class IConcertSession;
+struct FConcertBatchReplicationEvent;
+struct FConcertObjectReplicationEvent;
+struct FConcertSessionContext;
+
+namespace UE::ConcertSyncCore
+{
+	class FObjectReplicationCache;
+	
+	/** Receives replicated object data from all message endpoints and stores it in a FObjectReplicationCache. */
+	class CONCERTSYNCCORE_API FObjectReplicationReceiver
+	{
+	public:
+
+		FObjectReplicationReceiver(TSharedRef<IConcertSession> Session, TSharedRef<FObjectReplicationCache> ReplicationCache);
+		~FObjectReplicationReceiver();
+
+	private:
+
+		/** The session that is being received on. */
+		TSharedRef<IConcertSession> Session;
+		/** Where received data is stored. */
+		TSharedRef<FObjectReplicationCache> ReplicationCache;
+
+		void HandleBatchReplicationEvent(const FConcertSessionContext& SessionContext, const FConcertBatchReplicationEvent& Event);
+	};
+}
diff --git a/Engine/Plugins/Developer/Concert/ConcertSync/ConcertSyncCore/Source/ConcertSyncCore/Public/Replication/Processing/ObjectReplicationSender.h b/Engine/Plugins/Developer/Concert/ConcertSync/ConcertSyncCore/Source/ConcertSyncCore/Public/Replication/Processing/ObjectReplicationSender.h
new file mode 100644
--- /dev/null
+++ b/Engine/Plugins/Developer/Concert/ConcertSync/ConcertSyncCore/Source/ConcertSyncCore/Public/Replication/Processing/ObjectReplicationSender.h
@@ -0,0 +1,42 @@
+// Copyright Epic Games, Inc. All Rights Reserved.
+
+#pragma once
+
+#include "ObjectReplicationProcessor.h"
+#include "Replication/Messages/ConcertReplicationEvents.h"
+
+class IConcertSession;
+
+namespace UE::ConcertSyncCore
+{
+	/**
+	 * Sends object data to a specific endpoint ID.
+	 */
+	class CONCERTSYNCCORE_API FObjectReplicationSender : public FObjectReplicationProcessor
+	{
+	public:
+
+		FObjectReplicationSender(const FGuid& TargetEndpointId, TSharedRef<IConcertSession> Session, TSharedRef<IReplicationDataSource> DataSource);
+		
+		//~ Begin FObjectReplicationProcessor Interface
+		virtual void ProcessObjects(float TimeBudget) override;
+		//~ End FObjectReplicationProcessor Interface
+
+	protected:
+
+		//~ Begin FObjectReplicationProcessor Interface
+		virtual void ProcessObject(const FObjectProcessArgs& Args) override;
+		//~ End FObjectReplicationProcessor Interface
+
+	private:
+
+		/** The endpoint data will be sent to */
+		const FGuid TargetEndpointId;
+
+		/** The session through which replication messages are sent. */
+		TSharedRef<IConcertSession> Session;
+
+		/** This event is filled in ProcessObjects and finally sent to TargetEndpointId. */
+		FConcertBatchReplicationEvent EventToSend;
+	};
+}
diff --git a/Engine/Plugins/Developer/Concert/ConcertSync/ConcertSyncCore/Source/ConcertSyncCore/Public/Replication/Processing/ReplicationDataQueuer.h b/Engine/Plugins/Developer/Concert/ConcertSync/ConcertSyncCore/Source/ConcertSyncCore/Public/Replication/Processing/ReplicationDataQueuer.h
new file mode 100644
--- /dev/null
+++ b/Engine/Plugins/Developer/Concert/ConcertSync/ConcertSyncCore/Source/ConcertSyncCore/Public/Replication/Processing/ReplicationDataQueuer.h
@@ -0,0 +1,52 @@
+// Copyright Epic Games, Inc. All Rights Reserved.
+
+#pragma once
+
+#include "IReplicationDataSource.h"
+#include "ObjectReplicationCache.h"
+
+namespace UE::ConcertSyncCore
+{
+	/**
+	 * Holds on to events received from remote endpoints and exposes them as a source.
+	 * 
+	 * Received events are received from a FObjectReplicationCache, which makes sure that event data is shared effectively
+	 * if you created multiple FReplicationDataQueuer based on the same cache. This is relevant server side, where a
+	 * FReplicationDataQueuer is created for each client.
+	 * 
+	 * Child classes need to implement IReplicationCacheUser::WantsToAcceptObject.
+	 * TODO: Remove object data if the object is is meant for becomes unavailable.
+	 */
+	class CONCERTSYNCCORE_API FReplicationDataQueuer
+		: public IReplicationDataSource
+		, public IReplicationCacheUser
+		, public TSharedFromThis<FReplicationDataQueuer>
+	{
+	public:
+
+		virtual ~FReplicationDataQueuer() override;
+
+		//~ Begin IReplicationDataSource Interface
+		virtual void ForEachPendingObject(TFunctionRef<void(const FReplicationStreamObjectID&)> ProcessItemFunc) const override;
+		virtual int32 NumObjects() const override;
+		virtual bool ExtractReplicationDataForObject(const FReplicationStreamObjectID& Object, TFunctionRef<void(const FConcertSessionSerializedPayload& Payload)> ProcessCopyable, TFunctionRef<void(FConcertSessionSerializedPayload&& Payload)> ProcessMoveable) override;
+		//~ End IReplicationDataSource Interface
+		
+		//~ Begin IReplicationCacheUser Interface
+		virtual void OnDataCached(const FReplicationStreamObjectID& Object, TSharedRef<const FConcertObjectReplicationEvent> Data) override;
+		//~ End IReplicationCacheUser Interface
+
+	protected:
+
+		/** Called by subclass factory functions. */
+		void BindToCache(TSharedRef<FObjectReplicationCache> InReplicationCache);
+		
+	private:
+
+		/** Stores events as they are received. */
+		TMap<FReplicationStreamObjectID, TSharedPtr<const FConcertObjectReplicationEvent>> PendingEvents;
+
+		/** Provides us with replication events and shares them effectively. */
+		TSharedPtr<FObjectReplicationCache> ReplicationCache;
+	};
+}
diff --git a/Engine/Plugins/Developer/Concert/ConcertSync/ConcertSyncCore/Source/ConcertSyncCore/Public/Replication/PropertyChainUtils.h b/Engine/Plugins/Developer/Concert/ConcertSync/ConcertSyncCore/Source/ConcertSyncCore/Public/Replication/PropertyChainUtils.h
--- a/Engine/Plugins/Developer/Concert/ConcertSync/ConcertSyncCore/Source/ConcertSyncCore/Public/Replication/PropertyChainUtils.h
+++ b/Engine/Plugins/Developer/Concert/ConcertSync/ConcertSyncCore/Source/ConcertSyncCore/Public/Replication/PropertyChainUtils.h
@@ -68,6 +68,10 @@ namespace UE::ConcertSyncCore::PropertyChain
 
 	/** Util for determining whether this property is the inner property of a container. */
 	CONCERTSYNCCORE_API bool IsInnerContainerProperty(const FProperty& Property);
+
+	/** Checks whether this property would be marked as FConcertPropertyChain::InternalContainerPropertyValueName. */
+	CONCERTSYNCCORE_API bool IsPropertyEligibleForMarkingAsInternal(const FProperty& Property);
+
 	/**
 	 * Util for determining whether this property is considered primitive (numeric, bool, enum).
 	 * Primitive properties in containers are indicated as FConcertPropertyChain::InternalContainerPropertyValueName, which is "Value".
diff --git a/Engine/Plugins/Developer/Concert/ConcertSync/ConcertSyncCore/Source/ConcertSyncCore/Public/Replication/ReplicationPropertyFilter.h b/Engine/Plugins/Developer/Concert/ConcertSync/ConcertSyncCore/Source/ConcertSyncCore/Public/Replication/ReplicationPropertyFilter.h
new file mode 100644
--- /dev/null
+++ b/Engine/Plugins/Developer/Concert/ConcertSync/ConcertSyncCore/Source/ConcertSyncCore/Public/Replication/ReplicationPropertyFilter.h
@@ -0,0 +1,49 @@
+// Copyright Epic Games, Inc. All Rights Reserved.
+
+#pragma once
+
+#include "Containers/Array.h"
+#include "Containers/Map.h"
+#include "HAL/Platform.h"
+#include "UObject/NameTypes.h"
+
+class FProperty;
+struct FArchiveSerializedPropertyChain;
+struct FConcertPropertySelection;
+
+namespace UE::ConcertSyncCore
+{
+	/**
+	 * Util that bridges IObjectReplicationFormat::FAllowPropertyFunc and FConcertPropertySelection.
+	 *
+	 * Intended to be used like this:
+	 * {
+	 *		const FConcertPropertySelection& Selection = ...;
+	 *		IObjectReplicationFormat& Format = ...;
+	 *		UObject& ReplicatedObject = ...;
+	 *		
+	 *		FReplicationPropertyFilter Filter(Selection);
+	 *		Format.CreateReplicationEvent(Object, [&Filter](const FArchiveSerializedPropertyChain* Chain, const FProperty&& Property){ return Filter.IsPropertyInSelection(Chain, Property); }
+	 *	}
+	 */
+	class CONCERTSYNCCORE_API FReplicationPropertyFilter
+	{
+		const FConcertPropertySelection& PropertySelection;
+		
+		/**
+		 * Maps the name of every leaf property to the chain indices of PropertySelection.ReplicatedProperties that contain the property name at the end of the chain.
+		 * This speeds up matching chain and property.
+		 */
+		TMap<FName, TArray<int32>> LeafToChain;
+	public:
+
+		FReplicationPropertyFilter(const FConcertPropertySelection& PropertySelection);
+		
+		/**
+		 * @param Chain Can be nullptr or non-null and empty - both imply the Property is a root property.
+		 * @param Property The property that Chain leads to
+		 * @return Whether Property is in PropertySelection.
+		 */
+		bool ShouldSerializeProperty(const FArchiveSerializedPropertyChain* Chain, const FProperty& Property) const;
+	};
+}
\ No newline at end of file
diff --git a/Engine/Plugins/Developer/Concert/ConcertSync/ConcertSyncCore/Source/ConcertSyncCore/Public/Replication/ReplicationStreamObjectID.h b/Engine/Plugins/Developer/Concert/ConcertSync/ConcertSyncCore/Source/ConcertSyncCore/Public/Replication/ReplicationStreamObjectID.h
new file mode 100644
--- /dev/null
+++ b/Engine/Plugins/Developer/Concert/ConcertSync/ConcertSyncCore/Source/ConcertSyncCore/Public/Replication/ReplicationStreamObjectID.h
@@ -0,0 +1,29 @@
+// Copyright Epic Games, Inc. All Rights Reserved.
+
+#pragma once
+
+#include "Misc/Guid.h"
+#include "UObject/SoftObjectPath.h"
+
+namespace UE::ConcertSyncCore
+{
+	struct FReplicationStreamObjectID
+	{
+		/** The replication stream that produces data for this object. */
+		FGuid StreamId;
+		/** The object to process */
+		FSoftObjectPath Object;
+
+		friend bool operator==(const FReplicationStreamObjectID& Left, const FReplicationStreamObjectID& Right)
+		{
+			return Left.StreamId == Right.StreamId && Left.Object == Right.Object;
+		}
+
+		friend bool operator!=(const FReplicationStreamObjectID& Left, const FReplicationStreamObjectID& Right)
+		{
+			return !(Left == Right);
+		}
+	};
+
+	CONCERTSYNCCORE_API uint32 GetTypeHash(const FReplicationStreamObjectID& StreamObject);
+}
\ No newline at end of file
diff --git a/Engine/Plugins/Developer/Concert/ConcertSync/ConcertSyncServer/Source/ConcertSyncServer/Private/ConcertSyncServer.cpp b/Engine/Plugins/Developer/Concert/ConcertSync/ConcertSyncServer/Source/ConcertSyncServer/Private/ConcertSyncServer.cpp
--- a/Engine/Plugins/Developer/Concert/ConcertSync/ConcertSyncServer/Source/ConcertSyncServer/Private/ConcertSyncServer.cpp
+++ b/Engine/Plugins/Developer/Concert/ConcertSync/ConcertSyncServer/Source/ConcertSyncServer/Private/ConcertSyncServer.cpp
@@ -2,7 +2,6 @@
 
 #include "ConcertSyncServer.h"
 
-#include "IConcertModule.h"
 #include "IConcertServer.h"
 #include "IConcertSession.h"
 #include "ConcertServerWorkspace.h"
@@ -11,6 +10,7 @@
 #include "ConcertSyncServerArchivedSession.h"
 #include "ConcertSyncSessionDatabase.h"
 #include "ConcertLogGlobal.h"
+#include "Replication/ConcertServerReplicationManager.h"
 
 #include "Misc/Paths.h"
 #include "HAL/FileManager.h"
@@ -519,6 +519,17 @@ void FConcertSyncServer::DestroySequencerManager(const TSharedRef<FConcertSyncSe
 	LiveSessionSequencerManagers.Remove(InLiveSession->GetSession().GetId());
 }
 
+void FConcertSyncServer::CreateReplicationManager(TSharedRef<IConcertServerSession> InLiveSession)
+{
+	DestroyReplicationManager(InLiveSession);
+	LiveSessionReplicationManagers.Add(InLiveSession->GetId(), MakeShared<UE::ConcertSyncServer::Replication::FConcertServerReplicationManager>(MoveTemp(InLiveSession)));
+}
+
+void FConcertSyncServer::DestroyReplicationManager(const TSharedRef<IConcertServerSession>& InLiveSession)
+{
+	LiveSessionReplicationManagers.Remove(InLiveSession->GetId());
+}
+
 bool FConcertSyncServer::CreateLiveSession(const TSharedRef<IConcertServerSession>& InSession, const FInternalLiveSessionCreationParams& AdditionalParams)
 {
 	DestroyLiveSession(InSession);
@@ -533,6 +544,12 @@ bool FConcertSyncServer::CreateLiveSession(const TSharedRef<IConcertServerSessio
 			CreateSequencerManager(LiveSession.ToSharedRef());
 		}
 
+		// Create Replication Manager
+		if (EnumHasAnyFlags(LiveSession->GetSessionFlags(), EConcertSyncSessionFlags::EnableReplication))
+		{
+			CreateReplicationManager(InSession);
+		}
+
 		// We needn't call OnActivityProduced().Remove(...) because the subscription needs to stay for the lifetime of FConcertSyncServerLiveSession::SessionDatabase
 		check(AdditionalParams.OnModifiedCallback.IsBound());
 		LiveSession->GetSessionDatabase().OnActivityProduced().AddLambda([OnSessionModifiedCallback = AdditionalParams.OnModifiedCallback](const FConcertSyncActivity&)
diff --git a/Engine/Plugins/Developer/Concert/ConcertSync/ConcertSyncServer/Source/ConcertSyncServer/Private/ConcertSyncServer.h b/Engine/Plugins/Developer/Concert/ConcertSync/ConcertSyncServer/Source/ConcertSyncServer/Private/ConcertSyncServer.h
--- a/Engine/Plugins/Developer/Concert/ConcertSync/ConcertSyncServer/Source/ConcertSyncServer/Private/ConcertSyncServer.h
+++ b/Engine/Plugins/Developer/Concert/ConcertSync/ConcertSyncServer/Source/ConcertSyncServer/Private/ConcertSyncServer.h
@@ -13,6 +13,11 @@ class FConcertSyncServerLiveSession;
 class FConcertSyncServerArchivedSession;
 class FConcertSyncSessionDatabase;
 
+namespace UE::ConcertSyncServer::Replication
+{
+	class FConcertServerReplicationManager;
+}
+
 struct FConcertSessionFilter;
 
 /**
@@ -53,6 +58,9 @@ private:
 
 	void CreateSequencerManager(const TSharedRef<FConcertSyncServerLiveSession>& InLiveSession);
 	void DestroySequencerManager(const TSharedRef<FConcertSyncServerLiveSession>& InLiveSession);
+	
+	void CreateReplicationManager(TSharedRef<IConcertServerSession> InLiveSession);
+	void DestroyReplicationManager(const TSharedRef<IConcertServerSession>& InLiveSession); 
 
 	bool CreateLiveSession(const TSharedRef<IConcertServerSession>& InSession, const FInternalLiveSessionCreationParams& AdditionalParams);
 	void DestroyLiveSession(const TSharedRef<IConcertServerSession>& InSession);
@@ -74,6 +82,9 @@ private:
 	/** Map of live session IDs to their associated sequencer managers */
 	TMap<FGuid, TSharedPtr<FConcertServerSequencerManager>> LiveSessionSequencerManagers;
 
+	/** Map of live session IDs to their associated replication managers */
+	TMap<FGuid, TSharedRef<UE::ConcertSyncServer::Replication::FConcertServerReplicationManager>> LiveSessionReplicationManagers;
+
 	/** Map of live session IDs to their associated session data */
 	TMap<FGuid, TSharedPtr<FConcertSyncServerLiveSession>> LiveSessions;
 
diff --git a/Engine/Plugins/Developer/Concert/ConcertSync/ConcertSyncServer/Source/ConcertSyncServer/Private/Replication/ConcertReplicationClient.cpp b/Engine/Plugins/Developer/Concert/ConcertSync/ConcertSyncServer/Source/ConcertSyncServer/Private/Replication/ConcertReplicationClient.cpp
new file mode 100644
--- /dev/null
+++ b/Engine/Plugins/Developer/Concert/ConcertSync/ConcertSyncServer/Source/ConcertSyncServer/Private/Replication/ConcertReplicationClient.cpp
@@ -0,0 +1,26 @@
+// Copyright Epic Games, Inc. All Rights Reserved.
+
+#include "ConcertReplicationClient.h"
+
+#include "Processing/ServerReplicationDataQueuer.h"
+#include "Replication/Messages/ConcertReplicationHandshakeMessages.h"
+
+namespace UE::ConcertSyncServer::Replication
+{
+	FConcertReplicationClient::FConcertReplicationClient(
+		TArray<FReplicationStreamDescription> StreamDescriptions,
+		const FGuid& ClientEndpointId,
+		TSharedRef<IConcertSession> Session,
+		TSharedRef<ConcertSyncCore::FObjectReplicationCache> ReplicationCache
+	)
+		: StreamDescriptions(MoveTemp(StreamDescriptions))
+		, ClientEndpointId(ClientEndpointId)
+		, EventQueue(FServerReplicationDataQueuer::Make(*this, MoveTemp(ReplicationCache)))
+		, DataRelay(ClientEndpointId, MoveTemp(Session), EventQueue)
+	{}
+
+	void FConcertReplicationClient::ProcessClient(float TimeBudget)
+	{
+		DataRelay.ProcessObjects(TimeBudget);
+	}
+}
diff --git a/Engine/Plugins/Developer/Concert/ConcertSync/ConcertSyncServer/Source/ConcertSyncServer/Private/Replication/ConcertReplicationClient.h b/Engine/Plugins/Developer/Concert/ConcertSync/ConcertSyncServer/Source/ConcertSyncServer/Private/Replication/ConcertReplicationClient.h
new file mode 100644
--- /dev/null
+++ b/Engine/Plugins/Developer/Concert/ConcertSync/ConcertSyncServer/Source/ConcertSyncServer/Private/Replication/ConcertReplicationClient.h
@@ -0,0 +1,57 @@
+// Copyright Epic Games, Inc. All Rights Reserved.
+
+#pragma once
+
+#include "Containers/Array.h"
+#include "Replication/Processing/ObjectReplicationSender.h"
+#include "Templates/SharedPointer.h"
+
+struct FReplicationStreamDescription;
+struct FConcertReplication_Join_Request;
+
+namespace UE::ConcertSyncCore
+{
+	class FObjectReplicationCache;
+}
+
+namespace UE::ConcertSyncServer::Replication
+{
+	class FServerReplicationDataQueuer;
+
+	/** Server-side representation of a remote replication client. */
+	class FConcertReplicationClient : public FNoncopyable
+	{
+	public:
+
+		FConcertReplicationClient(
+			TArray<FReplicationStreamDescription> StreamDescriptions,
+			const FGuid& ClientEndpointId,
+			TSharedRef<IConcertSession> Session,
+			TSharedRef<ConcertSyncCore::FObjectReplicationCache> ReplicationCache
+			);
+
+		/**
+		 * Process any latent tasks, such as processing pending events that need to be sent to the remote instance.
+		 * Process given a time budget. The time budget may be exceeded but we'll try not to and to stay as close to the budget as possible.
+		 */
+		void ProcessClient(float TimeBudget);
+
+		const FGuid& GetClientEndpointId() const { return ClientEndpointId; }
+
+		const TArray<FReplicationStreamDescription>& GetStreamDescriptions() const { return StreamDescriptions; }
+
+	private:
+
+		/** The streams this client offered to send. */
+		const TArray<FReplicationStreamDescription> StreamDescriptions;
+
+		/** This client's endpoint ID. */
+		const FGuid ClientEndpointId;
+		
+		/** Queues up replication data and passes it to DataRelay. */
+		TSharedRef<FServerReplicationDataQueuer> EventQueue;
+
+		/** Sends data to the remote endpoint */
+		ConcertSyncCore::FObjectReplicationSender DataRelay;
+	};
+}
diff --git a/Engine/Plugins/Developer/Concert/ConcertSync/ConcertSyncServer/Source/ConcertSyncServer/Private/Replication/ConcertServerReplicationManager.cpp b/Engine/Plugins/Developer/Concert/ConcertSync/ConcertSyncServer/Source/ConcertSyncServer/Private/Replication/ConcertServerReplicationManager.cpp
new file mode 100644
--- /dev/null
+++ b/Engine/Plugins/Developer/Concert/ConcertSync/ConcertSyncServer/Source/ConcertSyncServer/Private/Replication/ConcertServerReplicationManager.cpp
@@ -0,0 +1,121 @@
+// Copyright Epic Games, Inc. All Rights Reserved.
+
+#include "ConcertServerReplicationManager.h"
+
+#include "ConcertLogGlobal.h"
+#include "IConcertSession.h"
+#include "Replication/ConcertReplicationClient.h"
+#include "Replication/Formats/FullObjectFormat.h"
+#include "Replication/Messages/ConcertReplicationHandshakeMessages.h"
+#include "Replication/Processing/ObjectReplicationCache.h"
+#include "Util/JoinRequestValidation.h"
+
+namespace UE::ConcertSyncServer::Replication
+{
+	FConcertServerReplicationManager::FConcertServerReplicationManager(TSharedRef<IConcertServerSession> InLiveSession)
+		: Session(MoveTemp(InLiveSession))
+		, ReplicationFormat(MakeShared<ConcertSyncCore::FFullObjectFormat>())
+		, ReplicationCache(MakeShared<ConcertSyncCore::FObjectReplicationCache>(ReplicationFormat))
+		, ReplicationDataReceiver(Session, ReplicationCache)
+	{
+		Session->RegisterCustomRequestHandler<FConcertReplication_Join_Request, FConcertReplication_Join_Response>(this, &FConcertServerReplicationManager::HandleJoinReplicationSessionRequest);
+		Session->RegisterCustomEventHandler<FConcertReplication_LeaveEvent>(this, &FConcertServerReplicationManager::HandleLeaveReplicationSessionRequest);
+		Session->OnSessionClientChanged().AddRaw(this, &FConcertServerReplicationManager::OnConnectionChanged);
+
+		Session->OnTick().AddRaw(this, &FConcertServerReplicationManager::Tick);
+	}
+
+	FConcertServerReplicationManager::~FConcertServerReplicationManager()
+	{
+		Session->UnregisterCustomRequestHandler<FConcertReplication_Join_Response>();
+		Session->UnregisterCustomEventHandler<FConcertReplication_LeaveEvent>(this);
+
+		Session->OnTick().RemoveAll(this);
+	}
+
+	EConcertSessionResponseCode FConcertServerReplicationManager::HandleJoinReplicationSessionRequest(
+		const FConcertSessionContext& ConcertSessionContext,
+		const FConcertReplication_Join_Request& Request,
+		FConcertReplication_Join_Response& Response
+		)
+	{
+		// Have a pair of logs before and after processing in case of potential disaster
+		UE_LOG(LogConcert, Log, TEXT("Received replication join request from endpoint %s"), *ConcertSessionContext.SourceEndpointId.ToString());
+		
+		const EConcertSessionResponseCode Result = InternalHandleJoinReplicationSessionRequest(ConcertSessionContext, Request, Response);
+		
+		UE_CLOG(Response.ErrorCode == EJoinReplicationErrorCode::Success, LogConcert, Log, TEXT("Accepted replication join request"));
+		UE_CLOG(Response.ErrorCode != EJoinReplicationErrorCode::Success, LogConcert, Log, TEXT("Rejected replication join request. %s: %s"), *ConcertSyncCore::Replication::LexJoinErrorCode(Response.ErrorCode), *Response.DetailedErrorMessage);
+		return Result;
+	}
+
+	EConcertSessionResponseCode FConcertServerReplicationManager::InternalHandleJoinReplicationSessionRequest(
+		const FConcertSessionContext& ConcertSessionContext,
+		const FConcertReplication_Join_Request& Request,
+		FConcertReplication_Join_Response& Response
+		)
+	{
+		const FGuid& ClientId = ConcertSessionContext.SourceEndpointId;
+		const bool bHasClient = Session->GetSessionClientEndpointIds().Contains(ClientId);
+		if (!bHasClient)
+		{
+			Response = { EJoinReplicationErrorCode::NotInAnyConcertSession, TEXT("Client must be in a Concert Session!") };
+			return EConcertSessionResponseCode::Success;
+		}
+
+		if (Clients.Contains(ClientId))
+		{
+			Response = { EJoinReplicationErrorCode::AlreadyInSession };
+			return EConcertSessionResponseCode::Success;
+		}
+
+		auto GetClientName = [this](const FGuid& ClientEndpointId) -> FString
+		{
+			FConcertSessionClientInfo Info;
+			if (!ensure(Session->FindSessionClient(ClientEndpointId, Info)))
+			{
+				return TEXT("Unknown Client");
+			}
+			return FString::Printf(TEXT("%s (%s)"), *Info.ClientInfo.UserName, *ClientEndpointId.ToString());
+		};
+		auto[ErrorCode, ErrorMessage, StreamDescriptions] = ValidateRequest(Request, Clients, GetClientName);
+		if (ErrorCode != EJoinReplicationErrorCode::Success)
+		{
+			Response = { ErrorCode, ErrorMessage };
+			return EConcertSessionResponseCode::Success;
+		}
+
+		Clients.Emplace(ClientId, MakeShared<FConcertReplicationClient>(MoveTemp(StreamDescriptions), ClientId, Session, ReplicationCache));
+		Response = { EJoinReplicationErrorCode::Success };
+		return EConcertSessionResponseCode::Success;
+	}
+
+	void FConcertServerReplicationManager::HandleLeaveReplicationSessionRequest(
+		const FConcertSessionContext& ConcertSessionContext,
+		const FConcertReplication_LeaveEvent& EventData
+		)
+	{
+		UE_LOG(LogConcert, Log, TEXT("Received replication leave request from endpoint %s"), *ConcertSessionContext.SourceEndpointId.ToString());
+		Clients.Remove(ConcertSessionContext.SourceEndpointId);
+	}
+
+	void FConcertServerReplicationManager::OnConnectionChanged(IConcertServerSession& ConcertServerSession, EConcertClientStatus ConcertClientStatus, const FConcertSessionClientInfo& ClientInfo)
+	{
+		if (ConcertClientStatus == EConcertClientStatus::Disconnected)
+		{
+			Clients.Remove(ClientInfo.ClientEndpointId);
+		}
+	}
+
+	void FConcertServerReplicationManager::Tick(IConcertServerSession& InSession, float InDeltaTime)
+	{
+		// TODO: Time slice replication clients
+		for (TPair<FGuid, TSharedRef<FConcertReplicationClient>>& Client : Clients)
+		{
+			// TODO: Load time budget from config
+			constexpr float TimeBudget = 1.f / 60.f;
+			Client.Value->ProcessClient(TimeBudget);
+		}
+	}
+}
+
diff --git a/Engine/Plugins/Developer/Concert/ConcertSync/ConcertSyncServer/Source/ConcertSyncServer/Private/Replication/ConcertServerReplicationManager.h b/Engine/Plugins/Developer/Concert/ConcertSync/ConcertSyncServer/Source/ConcertSyncServer/Private/Replication/ConcertServerReplicationManager.h
new file mode 100644
--- /dev/null
+++ b/Engine/Plugins/Developer/Concert/ConcertSync/ConcertSyncServer/Source/ConcertSyncServer/Private/Replication/ConcertServerReplicationManager.h
@@ -0,0 +1,62 @@
+// Copyright Epic Games, Inc. All Rights Reserved.
+
+#pragma once
+
+#include "ConcertMessages.h"
+#include "ConcertReplicationClient.h"
+#include "IConcertSessionHandler.h"
+#include "Replication/IConcertServerReplicationManager.h"
+#include "Replication/Messages/ConcertReplicationHandshakeMessages.h"
+#include "Replication/Processing/ObjectReplicationReceiver.h"
+#include "Templates/SharedPointer.h"
+#include "Templates/Tuple.h"
+
+namespace UE::ConcertSyncCore
+{
+	class IObjectReplicationFormat;
+	class FObjectReplicationCache;
+}
+
+class IConcertClientReplicationBridge;
+class IConcertServerSession;
+
+namespace UE::ConcertSyncServer::Replication
+{
+	class FConcertServerReplicationManager : public IConcertServerReplicationManager
+	{
+	public:
+		explicit FConcertServerReplicationManager(TSharedRef<IConcertServerSession> InLiveSession);
+		virtual ~FConcertServerReplicationManager() override;
+
+	private:
+		
+		/** Session instance this manager was created for. */
+		TSharedRef<IConcertServerSession> Session;
+		
+		/** Responsible for analysing received replication data. */
+		TSharedRef<ConcertSyncCore::IObjectReplicationFormat> ReplicationFormat;
+		
+		/** Received replication events are put into the ReplicationCache. The cache is used to relay data to clients latently. */
+		TSharedRef<ConcertSyncCore::FObjectReplicationCache> ReplicationCache;
+		/** Receives replication events from all endpoints. */
+		ConcertSyncCore::FObjectReplicationReceiver ReplicationDataReceiver;
+
+		/** Clients that have requested to join replication. Maps client ID to replication info. */
+		TMap<FGuid, TSharedRef<FConcertReplicationClient>> Clients;
+
+		// Event handlers
+		EConcertSessionResponseCode HandleJoinReplicationSessionRequest(const FConcertSessionContext& ConcertSessionContext, const FConcertReplication_Join_Request& Request, FConcertReplication_Join_Response& Response);
+		EConcertSessionResponseCode InternalHandleJoinReplicationSessionRequest(const FConcertSessionContext& ConcertSessionContext, const FConcertReplication_Join_Request& Request, FConcertReplication_Join_Response& Response);
+		void HandleLeaveReplicationSessionRequest(const FConcertSessionContext& ConcertSessionContext, const FConcertReplication_LeaveEvent& EventData);
+		
+		void OnConnectionChanged(IConcertServerSession& ConcertServerSession, EConcertClientStatus ConcertClientStatus, const FConcertSessionClientInfo& ConcertSessionClientInfo);
+
+		/**
+		 * Ticks all clients which causes clients to process pending data and send it to the corresponding endpoints.
+		 * 
+		 * There is an internal time budget to ensure ticks do not starve the server's other tick tasks.
+		 * This is configured in an .ini. TODO: Add config
+		 */
+		void Tick(IConcertServerSession& InSession, float InDeltaTime);
+	};
+}
\ No newline at end of file
diff --git a/Engine/Plugins/Developer/Concert/ConcertSync/ConcertSyncServer/Source/ConcertSyncServer/Private/Replication/Processing/ServerReplicationDataQueuer.cpp b/Engine/Plugins/Developer/Concert/ConcertSync/ConcertSyncServer/Source/ConcertSyncServer/Private/Replication/Processing/ServerReplicationDataQueuer.cpp
new file mode 100644
--- /dev/null
+++ b/Engine/Plugins/Developer/Concert/ConcertSync/ConcertSyncServer/Source/ConcertSyncServer/Private/Replication/Processing/ServerReplicationDataQueuer.cpp
@@ -0,0 +1,40 @@
+// Copyright Epic Games, Inc. All Rights Reserved.
+
+#include "ServerReplicationDataQueuer.h"
+
+#include "Replication/ConcertReplicationClient.h"
+#include "Replication/Data/ReplicationStreamDescription.h"
+
+namespace UE::ConcertSyncServer::Replication
+{
+	TSharedRef<FServerReplicationDataQueuer> FServerReplicationDataQueuer::Make(
+		const FConcertReplicationClient& OwningClient,
+		TSharedRef<ConcertSyncCore::FObjectReplicationCache> InReplicationCache
+		)
+	{
+		TSharedRef<FServerReplicationDataQueuer> Result = MakeShared<FServerReplicationDataQueuer>(OwningClient);
+		Result->BindToCache(MoveTemp(InReplicationCache));
+		return Result;
+	}
+
+	bool FServerReplicationDataQueuer::WantsToAcceptObject(const ConcertSyncCore::FReplicationStreamObjectID& Object) const
+	{
+		// Do not send back the data to the client that generated it
+		const bool bWasSentByThisClient = OwningClientStreamIds.Contains(Object.StreamId);
+		
+		// TODO: For now accept all objects from other clients. In the future, clients can specify what data they want to receive using client attributes.
+		return !bWasSentByThisClient;
+	}
+
+	FServerReplicationDataQueuer::FServerReplicationDataQueuer(const FConcertReplicationClient& OwningClient)
+		: OwningClientStreamIds([&OwningClient]()
+		{
+			TSet<FGuid> Streams;
+			for (const FReplicationStreamDescription& Description : OwningClient.GetStreamDescriptions())
+			{
+				Streams.Add(Description.BaseDescription.Identifier);
+			}
+			return Streams;
+		}())
+	{}
+}
diff --git a/Engine/Plugins/Developer/Concert/ConcertSync/ConcertSyncServer/Source/ConcertSyncServer/Private/Replication/Processing/ServerReplicationDataQueuer.h b/Engine/Plugins/Developer/Concert/ConcertSync/ConcertSyncServer/Source/ConcertSyncServer/Private/Replication/Processing/ServerReplicationDataQueuer.h
new file mode 100644
--- /dev/null
+++ b/Engine/Plugins/Developer/Concert/ConcertSync/ConcertSyncServer/Source/ConcertSyncServer/Private/Replication/Processing/ServerReplicationDataQueuer.h
@@ -0,0 +1,38 @@
+// Copyright Epic Games, Inc. All Rights Reserved.
+
+#pragma once
+
+#include "Containers/Set.h"
+#include "Misc/Guid.h"
+#include "Replication/Processing/ReplicationDataQueuer.h"
+#include "Templates/SharedPointer.h"
+
+namespace UE::ConcertSyncServer::Replication
+{
+	class FConcertReplicationClient;
+	
+	/**
+	 * Queues events for a specific client.
+	 * TODO: Filters events by matching client and stream attributes.
+	 */
+	class FServerReplicationDataQueuer : public ConcertSyncCore::FReplicationDataQueuer
+	{
+		template <typename ObjectType, ESPMode Mode>
+		friend class SharedPointerInternals::TIntrusiveReferenceController;
+	public:
+		
+		static TSharedRef<FServerReplicationDataQueuer> Make(const FConcertReplicationClient& OwningClient, TSharedRef<ConcertSyncCore::FObjectReplicationCache> InReplicationCache);
+		
+		//~ Begin IReplicationCacheUser Interface
+		virtual bool WantsToAcceptObject(const ConcertSyncCore::FReplicationStreamObjectID& Object) const override;
+		//~ End IReplicationCacheUser Interface
+
+	private:
+
+		FServerReplicationDataQueuer(const FConcertReplicationClient& OwningClient);
+
+		/** The client for which this FServerReplicationDataQueuer exists. */
+		const TSet<FGuid> OwningClientStreamIds;
+	};
+}
+
diff --git a/Engine/Plugins/Developer/Concert/ConcertSync/ConcertSyncServer/Source/ConcertSyncServer/Private/Replication/ReplicationTestInterface.cpp b/Engine/Plugins/Developer/Concert/ConcertSync/ConcertSyncServer/Source/ConcertSyncServer/Private/Replication/ReplicationTestInterface.cpp
new file mode 100644
--- /dev/null
+++ b/Engine/Plugins/Developer/Concert/ConcertSync/ConcertSyncServer/Source/ConcertSyncServer/Private/Replication/ReplicationTestInterface.cpp
@@ -0,0 +1,22 @@
+// Copyright Epic Games, Inc. All Rights Reserved.
+
+#include "ConcertServerReplicationManager.h"
+#include "Replication/IConcertServerReplicationManager.h"
+
+class IConcertServerSession;
+
+/**
+ * Exposes functions that are required for testing.
+ * 
+ * These functions are technically exported but conceptually not part of the public interface
+ * and should only be used for the purpose of automated testing.
+ */
+namespace UE::ConcertSyncServer::TestInterface
+{
+	CONCERTSYNCSERVER_API TSharedRef<Replication::IConcertServerReplicationManager> CreateServerReplicationManager(
+		TSharedRef<IConcertServerSession> InLiveSession
+		)
+	{
+		return MakeShared<Replication::FConcertServerReplicationManager>(MoveTemp(InLiveSession));
+	}
+}
diff --git a/Engine/Plugins/Developer/Concert/ConcertSync/ConcertSyncServer/Source/ConcertSyncServer/Private/Replication/Util/JoinRequestValidation.cpp b/Engine/Plugins/Developer/Concert/ConcertSync/ConcertSyncServer/Source/ConcertSyncServer/Private/Replication/Util/JoinRequestValidation.cpp
new file mode 100644
--- /dev/null
+++ b/Engine/Plugins/Developer/Concert/ConcertSync/ConcertSyncServer/Source/ConcertSyncServer/Private/Replication/Util/JoinRequestValidation.cpp
@@ -0,0 +1,178 @@
+// Copyright Epic Games, Inc. All Rights Reserved.
+
+#include "JoinRequestValidation.h"
+
+#include "Replication/ConcertReplicationClient.h"
+#include "Replication/Data/ObjectReplicationMap.h"
+#include "Replication/Data/ReplicationStreamDescription.h"
+#include "Replication/Messages/ConcertReplicationHandshakeMessages.h"
+
+namespace UE::ConcertSyncServer::Replication
+{
+	static TTuple<EJoinReplicationErrorCode, FString> CheckInputConstraints(
+		const FConcertReplication_Join_Request& Request,
+		const TMap<FGuid, TSharedRef<FConcertReplicationClient>>& Clients,
+		FGetClientName GetClientNameFunc
+		)
+	{
+		/* Duplicate objects and properties might not actually cause any trouble.
+		 * However, we will straight out reject any input that is not exactly how we expect it.
+		 * This serves us to
+		 *		1. reason about our pre/post conditions
+		 *		2. since other code can assume said conditions, it is a common security mechanic to reject input in case somebody sends us specifically crafted input
+		 */
+		TMap<FSoftObjectPath, TArray<const FReplicatedObjectInfo*>> UniqueObjectPropertyDetection;
+		TMap<FGuid, FConcertReplicationClient*> UniqueStreamIdDetection;
+		TSet<FGuid> UniqueStreamsFromRequestor;
+
+		// Prepare properties and stream IDs already used by other clients 
+		for (const TPair<FGuid, TSharedRef<FConcertReplicationClient>>& Client : Clients)
+		{
+			for (const FReplicationStreamDescription& StreamDescription : Client.Value->GetStreamDescriptions())
+			{
+				for (const TPair<FSoftObjectPath, FReplicatedObjectInfo>& ObjectSelection : StreamDescription.BaseDescription.ReplicationMap.ReplicatedObjects)
+				{
+					UniqueObjectPropertyDetection.FindOrAdd(ObjectSelection.Key).Add(&ObjectSelection.Value);
+				}
+				
+				const FGuid& ExistingStreamId = StreamDescription.BaseDescription.Identifier;
+				check(!UniqueStreamIdDetection.Contains(ExistingStreamId))
+				UniqueStreamIdDetection.Add(ExistingStreamId, &Client.Value.Get());
+			}
+		}
+		
+		for (const FReplicationStreamDescription_NetPacked& Stream : Request.Streams)
+		{
+			const FGuid& StreamId = Stream.BaseDescription.Identifier;
+
+			// Another client already using this stream?
+			if (FConcertReplicationClient** OwningClient = UniqueStreamIdDetection.Find(StreamId))
+			{
+				const FString ErrorMessage = FString::Printf(
+					TEXT("Stream %s is already in use by client %s"),
+					*StreamId.ToString(),
+					*GetClientNameFunc((*OwningClient)->GetClientEndpointId())
+					);
+				return { EJoinReplicationErrorCode::ConflictingStreamId, ErrorMessage };
+			}
+
+			// Request contained stream ID twice?
+			if (UniqueStreamsFromRequestor.Contains(StreamId))
+			{
+				return { EJoinReplicationErrorCode::ConflictingStreamId,  FString::Printf(TEXT("Duplicate stream ID %s in request."), *StreamId.ToString()) };
+			}
+			UniqueStreamsFromRequestor.Add(Stream.BaseDescription.Identifier);
+
+			// Make sure the properties are valid
+			for (const TPair<FSoftObjectPath, FReplicatedObjectInfo>& ReplicatedObjectInfo : Stream.BaseDescription.ReplicationMap.ReplicatedObjects)
+			{
+				const FSoftObjectPath& ObjectPath = ReplicatedObjectInfo.Key;
+				const FReplicatedObjectInfo& ObjectInfo = ReplicatedObjectInfo.Value;
+				
+				// Validate class - we cannot validate at this point whether the class really exists.
+				if (!ObjectInfo.ClassPath.IsValid())
+				{
+					const FString ErrorMessage = FString::Printf(
+						TEXT("Stream %s has null class for object %s"),
+						*Stream.BaseDescription.Identifier.ToString(),
+						*ObjectPath.ToString()
+						);
+					return { EJoinReplicationErrorCode::InvalidClass, ErrorMessage };
+				}
+
+				// Validate unique properties
+				TSet<FConcertPropertyChain> UniquePropertyDetection;
+				const TArray<const FReplicatedObjectInfo*>* PreexistingPropertySelections = UniqueObjectPropertyDetection.Find(ObjectPath);
+				for (const FConcertPropertyChain& PropertyChain : ObjectInfo.PropertySelection.ReplicatedProperties)
+				{
+					// No duplicate properties!
+					if (UniquePropertyDetection.Contains(PropertyChain))
+					{
+						const FString ErrorMessage = FString::Printf(
+							TEXT("Stream %s has duplicate property %s for object %s"),
+							*Stream.BaseDescription.Identifier.ToString(),
+							*PropertyChain.ToString(),
+							*ObjectPath.ToString()
+						);
+						return { EJoinReplicationErrorCode::DuplicateProperty, ErrorMessage };
+					}
+					UniquePropertyDetection.Add(PropertyChain);
+
+					if (!PreexistingPropertySelections)
+					{
+						continue;
+					}
+
+					// Is another client already streaming these properties?
+					for (const FReplicatedObjectInfo* PreexistingObjectSelection : *PreexistingPropertySelections)
+					{
+						// The request's object class should be the same the class previously assigned to the object in previous requests
+						if (PreexistingObjectSelection->ClassPath != ObjectInfo.ClassPath)
+						{
+							const FString ErrorMessage = FString::Printf(
+								TEXT("Class %s in stream %s for object %s does not match expected class %s (another client has told us this object has this class)."),
+								*ObjectInfo.ClassPath.ToString(),
+								*Stream.BaseDescription.Identifier.ToString(),
+								*ObjectPath.ToString(),
+								*PreexistingObjectSelection->ClassPath.ToString()
+							);
+							return { EJoinReplicationErrorCode::InvalidClass, ErrorMessage};
+						}
+
+						// Another client (or config in this request) is already using this property
+						if (PreexistingObjectSelection->PropertySelection.ReplicatedProperties.Contains(PropertyChain))
+						{
+							// TODO: In the future we'll change the handshake to no longer reject when another client wants to replicate this property. Adjust the message then.
+							// (Instead every client will request authority for each object before replicating)
+							const FString ErrorMessage = FString::Printf(
+								TEXT("Property %s in stream %s for object %s is already in duplicate or in use by another client"),
+								*PropertyChain.ToString(),
+								*Stream.BaseDescription.Identifier.ToString(),
+								*ObjectPath.ToString()
+							);
+							return { EJoinReplicationErrorCode::ConflictingAuthority, ErrorMessage };
+						}
+					}
+				}
+				
+				UniqueObjectPropertyDetection.FindOrAdd(ObjectPath).Add(&ObjectInfo);
+			}
+		}
+
+		return { EJoinReplicationErrorCode::Success, TEXT("") };
+	}
+	
+	TTuple<EJoinReplicationErrorCode, FString, TArray<FReplicationStreamDescription>> ValidateRequest(
+		const FConcertReplication_Join_Request& Request,
+		const TMap<FGuid, TSharedRef<FConcertReplicationClient>>& Clients,
+		FGetClientName GetClientNameFunc
+		)
+	{
+		if (auto[ErrorCode, ErrorMessage] = CheckInputConstraints(Request, Clients, GetClientNameFunc)
+			; ErrorCode != EJoinReplicationErrorCode::Success)
+		{
+			return { ErrorCode, ErrorMessage, TArray<FReplicationStreamDescription>{} };
+		}
+		
+		TArray<FReplicationStreamDescription> ParsedStreamDescriptions;
+		ParsedStreamDescriptions.Reserve(Request.Streams.Num());
+		for (const FReplicationStreamDescription_NetPacked& StreamDescription_NetPacked : Request.Streams)
+		{
+			// The net packed version contains UObject attributes, for which we may not be able to look up the UClass for
+			if (TOptional<FReplicationStreamDescription> Description = StreamDescription_NetPacked.Unpack())
+			{
+				ParsedStreamDescriptions.Emplace(MoveTemp(*Description));
+			}
+			else
+			{
+				return {
+					EJoinReplicationErrorCode::FailedToUnpackStream,
+					FString::Printf(TEXT("Failed to unpack stream %s"), *StreamDescription_NetPacked.BaseDescription.Identifier.ToString()),
+					TArray<FReplicationStreamDescription>{}
+				};
+			}
+		}
+
+		return { EJoinReplicationErrorCode::Success, TEXT(""), ParsedStreamDescriptions };
+	}
+}
diff --git a/Engine/Plugins/Developer/Concert/ConcertSync/ConcertSyncServer/Source/ConcertSyncServer/Private/Replication/Util/JoinRequestValidation.h b/Engine/Plugins/Developer/Concert/ConcertSync/ConcertSyncServer/Source/ConcertSyncServer/Private/Replication/Util/JoinRequestValidation.h
new file mode 100644
--- /dev/null
+++ b/Engine/Plugins/Developer/Concert/ConcertSync/ConcertSyncServer/Source/ConcertSyncServer/Private/Replication/Util/JoinRequestValidation.h
@@ -0,0 +1,37 @@
+// Copyright Epic Games, Inc. All Rights Reserved.
+
+#pragma once
+
+#include "Containers/Array.h"
+#include "Templates/SharedPointer.h"
+#include "Templates/Tuple.h"
+
+enum class EJoinReplicationErrorCode : uint8;
+struct FConcertReplication_Join_Request;
+struct FReplicationStreamDescription;
+
+namespace UE::ConcertSyncServer::Replication
+{
+	class FConcertReplicationClient;
+
+	using FGetClientName = TFunctionRef<FString(const FGuid& ClientEndpointID)>;
+	
+	/**
+	 * Attempts to unpack the request and validates it.
+	 * Constraints:
+	 *  - All stream IDs must be unique (within request and clients)
+	 *  - No overlapping properties (however, multiple streams can replicate the same object)
+	 *  - The stream descriptions must unpack successfully (i.e. the attributes must be known)
+	 * 
+	 * @param Request The request to validate
+	 * @param Clients Clients that are already connected
+	 * @param GetClientNameFunc Callback for getting a client's display name for error reporting purposes.
+	 * 
+	 * @return The error code for the request and the unpacked stream descriptions
+	 */
+	TTuple<EJoinReplicationErrorCode, FString, TArray<FReplicationStreamDescription>> ValidateRequest(
+		const FConcertReplication_Join_Request& Request,
+		const TMap<FGuid, TSharedRef<FConcertReplicationClient>>& Clients,
+		FGetClientName GetClientNameFunc
+		);
+}
diff --git a/Engine/Plugins/Developer/Concert/ConcertSync/ConcertSyncServer/Source/ConcertSyncServer/Public/Replication/IConcertServerReplicationManager.h b/Engine/Plugins/Developer/Concert/ConcertSync/ConcertSyncServer/Source/ConcertSyncServer/Public/Replication/IConcertServerReplicationManager.h
new file mode 100644
--- /dev/null
+++ b/Engine/Plugins/Developer/Concert/ConcertSync/ConcertSyncServer/Source/ConcertSyncServer/Public/Replication/IConcertServerReplicationManager.h
@@ -0,0 +1,19 @@
+// Copyright Epic Games, Inc. All Rights Reserved.
+
+#pragma once
+
+namespace UE::ConcertSyncServer::Replication
+{
+	/**
+	 * Handles all communication with clients regarding replication.
+	 * 
+	 * Negotiates replicated properties with clients, decides who is the authority (which is the first client to register properties),
+	 * and replicates received data to other clients.
+	 */
+	class IConcertServerReplicationManager
+	{
+	public:
+
+		virtual ~IConcertServerReplicationManager() = default;
+	};
+}
\ No newline at end of file
diff --git a/Engine/Plugins/Developer/Concert/ConcertSync/ConcertSyncTest/Source/ConcertSyncTest/ConcertSyncTest.Build.cs b/Engine/Plugins/Developer/Concert/ConcertSync/ConcertSyncTest/Source/ConcertSyncTest/ConcertSyncTest.Build.cs
--- a/Engine/Plugins/Developer/Concert/ConcertSync/ConcertSyncTest/Source/ConcertSyncTest/ConcertSyncTest.Build.cs
+++ b/Engine/Plugins/Developer/Concert/ConcertSync/ConcertSyncTest/Source/ConcertSyncTest/ConcertSyncTest.Build.cs
@@ -24,6 +24,7 @@ namespace UnrealBuildTool.Rules
 					"ConcertSyncCore",
 					"ConcertSyncClient",
 					"ConcertSyncServer",
+					"Engine"
 				}
 			);
 
diff --git a/Engine/Plugins/Developer/Concert/ConcertSync/ConcertSyncTest/Source/ConcertSyncTest/Private/Replication/ConcertPropertyTests.cpp b/Engine/Plugins/Developer/Concert/ConcertSync/ConcertSyncTest/Source/ConcertSyncTest/Private/Replication/ConcertPropertyTests.cpp
--- a/Engine/Plugins/Developer/Concert/ConcertSync/ConcertSyncTest/Source/ConcertSyncTest/Private/Replication/ConcertPropertyTests.cpp
+++ b/Engine/Plugins/Developer/Concert/ConcertSync/ConcertSyncTest/Source/ConcertSyncTest/Private/Replication/ConcertPropertyTests.cpp
@@ -2,10 +2,14 @@
 
 #include "Replication/Data/ConcertPropertySelection.h"
 #include "Replication/PropertyChainUtils.h"
-#include "Misc/AutomationTest.h"
-
+#include "Replication/ReplicationPropertyFilter.h"
 #include "TestReflectionObject.h"
+
+#include "Misc/AutomationTest.h"
 #include "Serialization/ArchiveSerializedPropertyChain.h"
+#include "Serialization/ObjectAndNameAsStringProxyArchive.h"
+#include "Serialization/ObjectReader.h"
+#include "Serialization/ObjectWriter.h"
 
 /**
  * Tests that the FConcertPropertyChain constructor works as the documentation dictates:
@@ -14,7 +18,7 @@
  *		- "normal" structs, like FVector > lists sub-properties
  *		- primitives and native structs > Path contains FConcertPropertyChain::InternalContainerPropertyValueName ("Value") at the end
  */
-IMPLEMENT_SIMPLE_AUTOMATION_TEST(FConcertPropertyChainTests, "Concert.Replication.PropertyChainTests", EAutomationTestFlags::EditorContext | EAutomationTestFlags::EngineFilter)
+IMPLEMENT_SIMPLE_AUTOMATION_TEST(FConcertPropertyChainTests, "Concert.Replication.Data.PropertyChainTests", EAutomationTestFlags::EditorContext | EAutomationTestFlags::EngineFilter)
 bool FConcertPropertyChainTests::RunTest(const FString& Parameters)
 {
 	// The test is intentionally set up to make it it easy to set breakpoints albeit at the expense of increasing the amount of code
@@ -114,7 +118,7 @@ bool FConcertPropertyChainTests::RunTest(const FString& Parameters)
 /**
  * Tests that all property cases on the specially constructed example class UTestReflectionObject.
  */
-IMPLEMENT_SIMPLE_AUTOMATION_TEST(FConcertForEachReplictableConcertPropertyTests, "Concert.Replication.ForEachReplicatableConcertProperty", EAutomationTestFlags::EditorContext | EAutomationTestFlags::EngineFilter)
+IMPLEMENT_SIMPLE_AUTOMATION_TEST(FConcertForEachReplictableConcertPropertyTests, "Concert.Replication.Data.ForEachReplicatableConcertProperty", EAutomationTestFlags::EditorContext | EAutomationTestFlags::EngineFilter)
 bool FConcertForEachReplictableConcertPropertyTests::RunTest(const FString& Parameters)
 {
 	UClass* Class = UTestReflectionObject::StaticClass();
@@ -167,13 +171,39 @@ bool FConcertForEachReplictableConcertPropertyTests::RunTest(const FString& Para
 		{ TEXT("TestStruct"), TEXT("Nested"), TEXT("NativeStructSet") },
 		{ TEXT("TestStruct"), TEXT("Nested"), TEXT("NativeStructSet"), FConcertPropertyChain::InternalContainerPropertyValueName },
 		{ TEXT("TestStruct"), TEXT("Nested"), TEXT("StringToFloat") },
-		{ TEXT("TestStruct"), TEXT("Nested"), TEXT("StringToFloat") },
+		{ TEXT("TestStruct"), TEXT("Nested"), TEXT("StringToFloat"), FConcertPropertyChain::InternalContainerPropertyValueName},
 		{ TEXT("TestStruct"), TEXT("Nested"), TEXT("StringToVector") },
 		{ TEXT("TestStruct"), TEXT("Nested"), TEXT("StringToVector"), TEXT("X") },
 		{ TEXT("TestStruct"), TEXT("Nested"), TEXT("StringToVector"), TEXT("Y") },
 		{ TEXT("TestStruct"), TEXT("Nested"), TEXT("StringToVector"), TEXT("Z") },
 		{ TEXT("TestStruct"), TEXT("Nested"), TEXT("StringToNativeStruct") },
 		{ TEXT("TestStruct"), TEXT("Nested"), TEXT("StringToNativeStruct"), FConcertPropertyChain::InternalContainerPropertyValueName },
+		// The same again for NestedArray as for Nested just above 
+		{ TEXT("TestStruct"), TEXT("NestedArray") },
+		{ TEXT("TestStruct"), TEXT("NestedArray"), TEXT("Value") },
+		{ TEXT("TestStruct"), TEXT("NestedArray"), TEXT("Vector") },
+		{ TEXT("TestStruct"), TEXT("NestedArray"), TEXT("Vector"), TEXT("X")  },
+		{ TEXT("TestStruct"), TEXT("NestedArray"), TEXT("Vector"), TEXT("Y")  },
+		{ TEXT("TestStruct"), TEXT("NestedArray"), TEXT("Vector"), TEXT("Z")  },
+		{ TEXT("TestStruct"), TEXT("NestedArray"), TEXT("NativeStruct") },
+		{ TEXT("TestStruct"), TEXT("NestedArray"), TEXT("FloatArray") },
+		{ TEXT("TestStruct"), TEXT("NestedArray"), TEXT("FloatArray"), FConcertPropertyChain::InternalContainerPropertyValueName },
+		{ TEXT("TestStruct"), TEXT("NestedArray"), TEXT("StringArray") },
+		{ TEXT("TestStruct"), TEXT("NestedArray"), TEXT("NativeStructArray") },
+		{ TEXT("TestStruct"), TEXT("NestedArray"), TEXT("NativeStructArray"), FConcertPropertyChain::InternalContainerPropertyValueName },
+		{ TEXT("TestStruct"), TEXT("NestedArray"), TEXT("FloatSet") },
+		{ TEXT("TestStruct"), TEXT("NestedArray"), TEXT("FloatSet"), FConcertPropertyChain::InternalContainerPropertyValueName },
+		{ TEXT("TestStruct"), TEXT("NestedArray"), TEXT("StringSet") },
+		{ TEXT("TestStruct"), TEXT("NestedArray"), TEXT("NativeStructSet") },
+		{ TEXT("TestStruct"), TEXT("NestedArray"), TEXT("NativeStructSet"), FConcertPropertyChain::InternalContainerPropertyValueName },
+		{ TEXT("TestStruct"), TEXT("NestedArray"), TEXT("StringToFloat") },
+		{ TEXT("TestStruct"), TEXT("NestedArray"), TEXT("StringToFloat"), FConcertPropertyChain::InternalContainerPropertyValueName},
+		{ TEXT("TestStruct"), TEXT("NestedArray"), TEXT("StringToVector") },
+		{ TEXT("TestStruct"), TEXT("NestedArray"), TEXT("StringToVector"), TEXT("X") },
+		{ TEXT("TestStruct"), TEXT("NestedArray"), TEXT("StringToVector"), TEXT("Y") },
+		{ TEXT("TestStruct"), TEXT("NestedArray"), TEXT("StringToVector"), TEXT("Z") },
+		{ TEXT("TestStruct"), TEXT("NestedArray"), TEXT("StringToNativeStruct") },
+		{ TEXT("TestStruct"), TEXT("NestedArray"), TEXT("StringToNativeStruct"), FConcertPropertyChain::InternalContainerPropertyValueName }
 	};
 
 	for (const FExpectedPropertyPath& Expected : ExpectedProperties)
@@ -184,7 +214,257 @@ bool FConcertForEachReplictableConcertPropertyTests::RunTest(const FString& Para
 			AddError(ToString(Expected));
 		}
 	}
-	TestEqual(TEXT("Does not contain disallowed properties"), Properties.Num(), ExpectedProperties.Num());
 	
+	if (Properties.Num() != ExpectedProperties.Num())
+	{
+		AddError(TEXT("More properties found than expected"));
+		for (const FConcertPropertyChain& Chain : Properties)
+		{
+			if (!ExpectedProperties.Contains(Chain))
+			{
+				AddError(FString::Printf(TEXT("Found unexpected chain %s"), *Chain.ToString()));
+			}
+		}
+	}
+	
+	return true;
+}
+
+/**
+ * Tests FConcertPropertyChain::Matches.
+ */
+IMPLEMENT_SIMPLE_AUTOMATION_TEST(FMatchConcertPropertyChainTest, "Concert.Replication.Data.MatchConcertPropertyChain", EAutomationTestFlags::EditorContext | EAutomationTestFlags::EngineFilter)
+bool FMatchConcertPropertyChainTest::RunTest(const FString& Parameters)
+{
+	// 1. Set up
+	UClass* Class = UTestReflectionObject::StaticClass();
+	FProperty* FloatProperty = Class->FindPropertyByName(GET_MEMBER_NAME_CHECKED(UTestReflectionObject, Float));
+	FProperty* TestStructProperty = Class->FindPropertyByName(GET_MEMBER_NAME_CHECKED(UTestReflectionObject, TestStruct));
+	FProperty* NestedArrayProperty = FTestNestedReplicationStruct::StaticStruct()->FindPropertyByName(GET_MEMBER_NAME_CHECKED(FTestNestedReplicationStruct, NestedArray));
+	FProperty* InternalNestedArrayProperty = CastField<FArrayProperty>(NestedArrayProperty)->Inner;
+	FProperty* ValueProperty = FTestReplicationStruct::StaticStruct()->FindPropertyByName(GET_MEMBER_NAME_CHECKED(FTestReplicationStruct, Value));
+	FProperty* FloatArrayProperty = FTestReplicationStruct::StaticStruct()->FindPropertyByName(GET_MEMBER_NAME_CHECKED(FTestReplicationStruct, FloatArray));
+	FProperty* InternalFloatArrayProperty = CastField<FArrayProperty>(FloatArrayProperty)->Inner;
+	FProperty* StringToFloatProperty = FTestReplicationStruct::StaticStruct()->FindPropertyByName(GET_MEMBER_NAME_CHECKED(FTestReplicationStruct, StringToFloat));
+	FProperty* StringToFloatValueProperty = CastField<FMapProperty>(StringToFloatProperty)->ValueProp;
+	
+	FArchiveSerializedPropertyChain ChainToNestedArrayProperty;
+	ChainToNestedArrayProperty.PushProperty(TestStructProperty, false);
+	// FTestReplicationStruct::Value (float)
+	FArchiveSerializedPropertyChain ChainToNestedProperty = ChainToNestedArrayProperty;
+	// FTestReplicationStruct::FloatArray (TArray<float>)
+	ChainToNestedProperty.PushProperty(NestedArrayProperty, false);
+	ChainToNestedProperty.PushProperty(InternalNestedArrayProperty, false);
+	FArchiveSerializedPropertyChain ChainToFloatArrayProperty = ChainToNestedProperty;
+	FArchiveSerializedPropertyChain ChainToInternalFloatArrayProperty = ChainToFloatArrayProperty;
+	ChainToInternalFloatArrayProperty.PushProperty(FloatArrayProperty, false);
+	// FTestReplicationStruct::StringToFloat (TMap<FString, float>)
+	FArchiveSerializedPropertyChain ChainToStringToFloatProperty = ChainToNestedProperty;
+	FArchiveSerializedPropertyChain ChainToStringToFloatValueProperty = ChainToNestedProperty;
+	ChainToStringToFloatValueProperty.PushProperty(StringToFloatProperty, false);
+	
+	const FConcertPropertyChain ConcertChain_Simple_FloatProperty(nullptr, *FloatProperty);
+	const FConcertPropertyChain ConcertChain_Nested_ArrayProperty(&ChainToNestedArrayProperty, *NestedArrayProperty);
+	const FConcertPropertyChain ConcertChain_Nested_ValueProperty(&ChainToNestedProperty, *ValueProperty);
+	const FConcertPropertyChain ConcertChain_Nested_FloatArrayProperty(&ChainToFloatArrayProperty, *FloatArrayProperty);
+	const FConcertPropertyChain ConcertChain_Nested_FloatArrayInnerProperty(&ChainToInternalFloatArrayProperty, *InternalFloatArrayProperty);
+	const FConcertPropertyChain ConcertChain_Nested_StringToFloatProperty(&ChainToStringToFloatProperty, *StringToFloatProperty);
+	const FConcertPropertyChain ConcertChain_Nested_StringToFloatValueProperty(&ChainToStringToFloatValueProperty, *StringToFloatValueProperty);
+	
+
+
+	// 2. Run
+	const bool bMatches_Simple_FloatProperty = ConcertChain_Simple_FloatProperty.MatchesExactly(nullptr, *FloatProperty);
+	const bool bMatches_Nested_ArrayProperty = ConcertChain_Nested_ArrayProperty.MatchesExactly(&ChainToNestedArrayProperty, *NestedArrayProperty);
+	const bool bMatches_Nested_Property = ConcertChain_Nested_ValueProperty.MatchesExactly(&ChainToNestedProperty, *ValueProperty);
+	const bool bMatches_Nested_FloatArrayProperty = ConcertChain_Nested_FloatArrayProperty.MatchesExactly(&ChainToFloatArrayProperty, *FloatArrayProperty);
+	const bool bMatches_Nested_FloatArrayInnerProperty = ConcertChain_Nested_FloatArrayInnerProperty.MatchesExactly(&ChainToInternalFloatArrayProperty, *InternalFloatArrayProperty);
+	const bool bMatches_Nested_StringToFloatProperty = ConcertChain_Nested_StringToFloatProperty.MatchesExactly(&ChainToStringToFloatProperty, *StringToFloatProperty);
+	const bool bMatches_Nested_StringToFloatValueProperty = ConcertChain_Nested_StringToFloatValueProperty.MatchesExactly(&ChainToStringToFloatValueProperty, *StringToFloatValueProperty);
+
+
+
+
+	// 3. Test
+	TestTrue(TEXT("Float"), bMatches_Simple_FloatProperty);
+	TestTrue(TEXT("NestedArray"), bMatches_Nested_ArrayProperty);
+	// This case is interesting because e.g. ChainToNestedProperty = { NestedArray, NestedArray } but ConcertChain_NestedProperty = { NestedArray, Value } (set a breakpoint to convince yourself)
+	TestTrue(TEXT("NestedArray.NestedArray.Value"), bMatches_Nested_Property);
+	// This would restore the size of the array only...
+	TestTrue(TEXT("NestedArray.NestedArray.FloatArray"), bMatches_Nested_FloatArrayProperty);
+	// ... and this includes the float values. This one is interesting because it ends with FConcertPropertyChain::InternalContainerPropertyValueName
+	TestTrue(TEXT("NestedArray.NestedArray.FloatArray.Value (ends with FConcertPropertyChain::InternalContainerPropertyValueName)"), bMatches_Nested_FloatArrayInnerProperty);
+	// This would restore the keys and the size of the map only...
+	TestTrue(TEXT("NestedArray.NestedArray.StringToFloat"), bMatches_Nested_StringToFloatProperty);
+	// ... and this includes the float values. This one is interesting because it ends with FConcertPropertyChain::InternalContainerPropertyValueName
+	TestTrue(TEXT("NestedArray.NestedArray.StringToFloat.Value (ends with FConcertPropertyChain::InternalContainerPropertyValueName)"), bMatches_Nested_StringToFloatValueProperty);
+
+	return true;
+}
+
+/**
+ * Tests that UE::ConcertSyncCore::FReplicationPropertyFilter works correctly.
+ * This uses the filter and serializes an object, doing a Memcmp to detect that exactly the expected properties are serialized.
+ */
+IMPLEMENT_SIMPLE_AUTOMATION_TEST(FSerializeWithPropertyFilterTest, "Concert.Replication.Data.SerializeWithPropertyFilter", EAutomationTestFlags::EditorContext | EAutomationTestFlags::EngineFilter)
+bool FSerializeWithPropertyFilterTest::RunTest(const FString& Parameters)
+{
+	class FTestArchive : public FMemoryWriter
+	{
+		FAutomationTestBase& Test;
+		const FConcertPropertySelection& Selection;
+		const UE::ConcertSyncCore::FReplicationPropertyFilter PropertyFilter;
+	public:
+		
+		TArray<uint8> Bytes;
+
+		FTestArchive(FAutomationTestBase& Test, const FConcertPropertySelection& Selection)
+			: FMemoryWriter(Bytes)
+			, Test(Test)
+			, Selection(Selection)
+			, PropertyFilter(Selection)
+		{
+			FArchive::SetIsSaving(true);
+		}
+
+		virtual bool ShouldSkipProperty(const FProperty* InProperty) const override
+		{
+			if (!UE::ConcertSyncCore::PropertyChain::IsReplicatableProperty(*InProperty))
+			{
+				return true;
+			}
+
+			// We can assume the constructor of FConcertPropertyChain works correctly here since we've tested it in another test.
+			const FConcertPropertyChain Chain(GetSerializedPropertyChain(), *InProperty);
+			const bool bIsPropertyInChain = Selection.ReplicatedProperties.Contains(Chain);
+			const bool bIsAllowed = PropertyFilter.ShouldSerializeProperty(GetSerializedPropertyChain(), *InProperty);
+			if (bIsPropertyInChain)
+			{
+				Test.TestTrue(FString::Printf(TEXT("Allowed %s"), *InProperty->GetAuthoredName()), bIsAllowed);
+			}
+			else
+			{
+				Test.TestFalse(FString::Printf(TEXT("Disallowed %s"), *InProperty->GetAuthoredName()), bIsAllowed);
+			}
+
+			return !bIsAllowed;
+		}
+
+		// If we encounter any objects, ring alarm bells because IsReplicatableProperty should have covered this case
+		virtual FArchive& operator<<(UObject*& Obj) override{ Test.AddError(TEXT("Object not expected at this point")); return *this; }
+		virtual FArchive& operator<<(FWeakObjectPtr& Obj) override { Test.AddError(TEXT("Object not expected at this point")); return *this; }
+		virtual FArchive& operator<<(FSoftObjectPtr& Value) override { Test.AddError(TEXT("Object not expected at this point")); return *this; } 
+		virtual FArchive& operator<<(FSoftObjectPath& Value) override { Test.AddError(TEXT("Object not expected at this point")); return *this; }
+		virtual FArchive& operator<<(FObjectPtr& Obj) override { Test.AddError(TEXT("Object not expected at this point")); return *this; }
+	};
+
+	// CreateFromPath is supposed to work (assuming the other tests worked so far) but we'll use TOptional::Get() calls are in case it does not. 
+	TArray<FConcertPropertyChain> AllowedProperties = {
+		FConcertPropertyChain::CreateFromPath(*UTestReflectionObject::StaticClass(), { TEXT("Float") }).Get({}),
+		FConcertPropertyChain::CreateFromPath(*UTestReflectionObject::StaticClass(), { TEXT("Vector") }).Get({}),
+		FConcertPropertyChain::CreateFromPath(*UTestReflectionObject::StaticClass(), { TEXT("Vector"), TEXT("X") }).Get({}),
+		FConcertPropertyChain::CreateFromPath(*UTestReflectionObject::StaticClass(), { TEXT("Vector"), TEXT("Y") }).Get({}),
+		FConcertPropertyChain::CreateFromPath(*UTestReflectionObject::StaticClass(), { TEXT("Vector"), TEXT("Z") }).Get({}),
+		FConcertPropertyChain::CreateFromPath(*UTestReflectionObject::StaticClass(), { TEXT("TestStruct") }).Get({}),
+		FConcertPropertyChain::CreateFromPath(*UTestReflectionObject::StaticClass(), { TEXT("TestStruct"), TEXT("Nested") }).Get({}),
+		FConcertPropertyChain::CreateFromPath(*UTestReflectionObject::StaticClass(), { TEXT("TestStruct"), TEXT("Nested"), TEXT("Value") }).Get({}),
+		FConcertPropertyChain::CreateFromPath(*UTestReflectionObject::StaticClass(), { TEXT("TestStruct"), TEXT("Nested"), TEXT("NativeStruct") }).Get({}),
+		FConcertPropertyChain::CreateFromPath(*UTestReflectionObject::StaticClass(), { TEXT("TestStruct"), TEXT("Nested"), TEXT("FloatArray") }).Get({}),
+		FConcertPropertyChain::CreateFromPath(*UTestReflectionObject::StaticClass(), { TEXT("TestStruct"), TEXT("Nested"), TEXT("FloatArray"), FConcertPropertyChain::InternalContainerPropertyValueName }).Get({}),
+		FConcertPropertyChain::CreateFromPath(*UTestReflectionObject::StaticClass(), { TEXT("TestStruct"), TEXT("NestedArray") }).Get({}),
+		FConcertPropertyChain::CreateFromPath(*UTestReflectionObject::StaticClass(), { TEXT("TestStruct"), TEXT("NestedArray"), TEXT("Value") }).Get({}),
+		FConcertPropertyChain::CreateFromPath(*UTestReflectionObject::StaticClass(), { TEXT("TestStruct"), TEXT("NestedArray"), TEXT("Vector") }).Get({}),
+		FConcertPropertyChain::CreateFromPath(*UTestReflectionObject::StaticClass(), { TEXT("TestStruct"), TEXT("NestedArray"), TEXT("Vector"), TEXT("X") }).Get({}),
+		FConcertPropertyChain::CreateFromPath(*UTestReflectionObject::StaticClass(), { TEXT("TestStruct"), TEXT("NestedArray"), TEXT("Vector"), TEXT("Y") }).Get({}),
+		FConcertPropertyChain::CreateFromPath(*UTestReflectionObject::StaticClass(), { TEXT("TestStruct"), TEXT("NestedArray"), TEXT("Vector"), TEXT("Z") }).Get({}),
+		FConcertPropertyChain::CreateFromPath(*UTestReflectionObject::StaticClass(), { TEXT("TestStruct"), TEXT("NestedArray"), TEXT("NativeStruct") }).Get({}),
+		FConcertPropertyChain::CreateFromPath(*UTestReflectionObject::StaticClass(), { TEXT("TestStruct"), TEXT("NestedArray"), TEXT("FloatArray") }).Get({}),
+		FConcertPropertyChain::CreateFromPath(*UTestReflectionObject::StaticClass(), { TEXT("TestStruct"), TEXT("NestedArray"), TEXT("FloatArray"), FConcertPropertyChain::InternalContainerPropertyValueName }).Get({})
+	};
+	if (AllowedProperties.Contains(FConcertPropertyChain{}))
+	{
+		AddError(TEXT("FConcertPropertyChain::CreateFromPath was expected to work on all of the above cases. Take a look at the ForEachReplicatableConcertProperty test, which should have failed, as well."));
+		return false;
+	}
+
+	auto SetTestValues = [](UTestReflectionObject& Object)
+	{
+		Object.Float = 42.f;
+		Object.Vector = { 10.f, 20.f, 30.f };
+		
+		Object.TestStruct.Nested.Value = 420.f;
+		Object.TestStruct.Nested.NativeStruct.Float = 4200.f;
+		Object.TestStruct.Nested.FloatArray = { 500.f, 750.f };
+
+		Object.TestStruct.NestedArray = { FTestReplicationStruct{} };
+		Object.TestStruct.NestedArray[0].Value = -420.f;
+		Object.TestStruct.NestedArray[0].Vector = { -100.f, -200.f, -300.f };
+		Object.TestStruct.NestedArray[0].NativeStruct.Float = -4200.f;
+		Object.TestStruct.NestedArray[0].FloatArray = { -500.f, -750.f };
+	};
+	auto ResetTestValues = [](UTestReflectionObject& Object)
+	{
+		Object.Float = {};
+		Object.Vector = FVector::ZeroVector;
+		Object.TestStruct = FTestNestedReplicationStruct{};
+	};
+	auto ObjectToBinary = [](UObject& Object)
+	{
+		TArray<uint8> Bytes;
+		FObjectWriter CollectAllData(&Object, Bytes, false);
+		return Bytes;
+	};
+	auto ValidateTestValues = [this](UTestReflectionObject& Object)
+	{
+		TestEqual(TEXT("Object.Float = 42.f"), Object.Float, 42.f);
+		TestEqual(TEXT("Object.Vector = { 10.f, 20.f, 30.f }"), Object.Vector, FVector{10.f, 20.f, 30.f });
+		
+		TestEqual(TEXT("Object.TestStruct.Nested.Value = 420.f"), Object.TestStruct.Nested.Value, 420.f);
+		TestEqual(TEXT("Object.TestStruct.Nested.NativeStruct.Float = 4200.f"), Object.TestStruct.Nested.NativeStruct.Float, 4200.f);
+		TestEqual(TEXT("Object.TestStruct.Nested.FloatArray = { 500.f, 750.f }"), Object.TestStruct.Nested.FloatArray, TArray<float>{ 500.f, 750.f });
+
+		TestEqual(TEXT("Object.TestStruct.NestedArray[0].Value = -420.f"), Object.TestStruct.NestedArray[0].Value, -420.f);
+		TestEqual(TEXT("Object.TestStruct.NestedArray[0].Vector = { -100.f, -200.f, -300.f }"), Object.TestStruct.NestedArray[0].Vector, FVector{ -100.f, -200.f, -300.f });
+		TestEqual(TEXT("Object.TestStruct.NestedArray[0].NativeStruct.Float = -4200.f"), Object.TestStruct.NestedArray[0].NativeStruct.Float, -4200.f);
+		TestEqual(TEXT("Object.TestStruct.NestedArray[0].FloatArray = { -500.f, -750.f }"), Object.TestStruct.NestedArray[0].FloatArray, TArray<float>{ -500.f, -750.f });
+	};
+	
+	UTestReflectionObject* TestObject = NewObject<UTestReflectionObject>();
+	SetTestValues(*TestObject);
+	
+	// This property is supposed to be filtered out...
+	TestObject->TestStruct.Nested.Vector = { 100.f, 200.f, 300.f };
+	const TArray<uint8> SnapshotWithTestValues = ObjectToBinary(*TestObject);
+	// ... and we change its value after taking snapshotting to validate the filter indeed filters out the undesired property...
+	TestObject->TestStruct.Nested.Vector = { 1000.f, 2000.f, 3000.f };
+
+	// This does the filtering & tests that properties are filtered correctly using ShouldSkipProperty
+	const FConcertPropertySelection Selection{ AllowedProperties };
+	FTestArchive SavingArchive(*this, Selection);
+	FNameAsStringProxyArchive CrashProtection_Saving(SavingArchive);
+	TestObject->Serialize(CrashProtection_Saving);
+
+	// ... the supposed to be filtered property must now be reset to the same value as was in the snapshot so Memcmp does not detect it.
+	ResetTestValues(*TestObject);
+	TestObject->TestStruct.Nested.Vector = { 100.f, 200.f, 300.f };
+	FObjectReader LoadingArchive(SavingArchive.Bytes);
+	FNameAsStringProxyArchive CrashProtection_Loading(LoadingArchive);
+	TestObject->Serialize(CrashProtection_Loading);
+	const TArray<uint8> SnapshotAfterLoading = ObjectToBinary(*TestObject);
+	
+	// We do a Memcmp on the snapshot data before filtering and on the data after applying to a blanked out object. They must be equal or the filter did not work correctly!
+	if (SnapshotWithTestValues.Num() == SnapshotAfterLoading.Num())
+	{
+		const bool bEqualMemory = FPlatformMemory::Memcmp(SnapshotWithTestValues.GetData(), SnapshotAfterLoading.GetData(), SnapshotWithTestValues.Num()) == 0;
+		TestTrue(TEXT("Memcmp returned non-zero: Applying the filtered data to a reset object did not leave the object in the same state as it was serialized in."), bEqualMemory);
+	}
+	else
+	{
+		AddError(TEXT("BinaryDataWithTestValues != BinaryDataAfterLoading: Applying the filtered data to a reset object did not leave the object in the same state as it was serialized in."));
+	}
+
+	// This is technically unnecessary because Memcmp would already have failed the test above but this helps in debugging what values caused the Memcmp to be different.
+	ValidateTestValues(*TestObject);
+	TestEqual(TEXT("Object.TestStruct.Nested.Vector was left unchanged"), TestObject->TestStruct.Nested.Vector, FVector{ 100.f, 200.f, 300.f });
 	return true;
 }
\ No newline at end of file
diff --git a/Engine/Plugins/Developer/Concert/ConcertSync/ConcertSyncTest/Source/ConcertSyncTest/Private/Replication/ReplicationCacheTests.cpp b/Engine/Plugins/Developer/Concert/ConcertSync/ConcertSyncTest/Source/ConcertSyncTest/Private/Replication/ReplicationCacheTests.cpp
new file mode 100644
--- /dev/null
+++ b/Engine/Plugins/Developer/Concert/ConcertSync/ConcertSyncTest/Source/ConcertSyncTest/Private/Replication/ReplicationCacheTests.cpp
@@ -0,0 +1,171 @@
+// Copyright Epic Games, Inc. All Rights Reserved.
+
+#include "Replication/Formats/IObjectReplicationFormat.h"
+#include "Replication/Messages/ConcertReplicationEvents.h"
+#include "Replication/Processing/ObjectReplicationCache.h"
+#include "Replication/ReplicationStreamObjectID.h"
+#include "TestReflectionObject.h"
+#include "Util/ConcertMocks.h"
+
+#include "Misc/AutomationTest.h"
+#include "Misc/Paths.h"
+#include "HAL/FileManager.h"
+
+namespace UE::ConcertSyncTests
+{
+	enum class EReplicationCacheTestFlags
+	{
+		None = 0,
+		InstantConsume = 1 << 0,
+		NeverConsume = 1 << 1,
+		NeverReceive = 1 << 2,
+	};
+	ENUM_CLASS_FLAGS(EReplicationCacheTestFlags);
+
+	/**
+	 * Tests that FObjectReplicationCache works as expected.
+	 */
+	IMPLEMENT_SIMPLE_AUTOMATION_TEST(FReplicationCacheTest, "Concert.Replication.Components.ReplicationCache", EAutomationTestFlags::EditorContext | EAutomationTestFlags::EngineFilter);
+	bool FReplicationCacheTest::RunTest(const FString& Parameters)
+	{
+		class FTestReplicationFormat : public ConcertSyncCore::IObjectReplicationFormat
+		{
+		public:
+			virtual TOptional<FConcertSessionSerializedPayload> CreateReplicationEvent(UObject& Object, FAllowPropertyFunc IsPropertyAllowedFunc) override { return NotMocked<TOptional<FConcertSessionSerializedPayload>>({}); }
+			virtual void ClearInternalCache(TArrayView<UObject> ObjectsToClear) override { return NotMocked<void>(); }
+			virtual void CombineReplicationEvents(FConcertSessionSerializedPayload& Base, const FConcertSessionSerializedPayload& Newer) override
+			{
+				// Reuse FNativeStruct to avoid introducing more test types
+				FNativeStruct BaseStruct;
+				Base.GetTypedPayload(BaseStruct);
+				FNativeStruct NewerStruct;
+				Newer.GetTypedPayload(NewerStruct);
+
+				BaseStruct.Float += NewerStruct.Float;
+				Base.SetTypedPayload(BaseStruct);
+			}
+			virtual void ApplyReplicationEvent(UObject& Object, const FConcertSessionSerializedPayload& Payload) override { return NotMocked<void>(); }
+
+			static FConcertObjectReplicationEvent CreateEvent(FSoftObjectPath Object, float Value)
+			{
+				FNativeStruct Data{ Value };
+				FConcertSessionSerializedPayload Payload;
+				Payload.SetTypedPayload(Data, EConcertPayloadCompressionType::None);
+				return { MoveTemp(Object), Payload };
+			}
+		};
+
+		class FTestReplicationCacheUser : public ConcertSyncCore::IReplicationCacheUser
+		{
+		public:
+
+			FAutomationTestBase& Test;
+			ConcertSyncCore::FReplicationStreamObjectID AllowedObject;
+			EReplicationCacheTestFlags Flags;
+			TSharedPtr<const FConcertObjectReplicationEvent> CachedData;
+
+			FTestReplicationCacheUser(FAutomationTestBase& Test, ConcertSyncCore::FReplicationStreamObjectID AllowedObject, EReplicationCacheTestFlags Flags = EReplicationCacheTestFlags::None)
+				: Test(Test)
+				, AllowedObject(MoveTemp(AllowedObject))
+				, Flags(Flags)
+			{}
+			
+			virtual bool WantsToAcceptObject(const ConcertSyncCore::FReplicationStreamObjectID& Object) const override
+			{
+				return !EnumHasAnyFlags(Flags, EReplicationCacheTestFlags::NeverReceive) && AllowedObject == Object;
+			}
+			
+			virtual void OnDataCached(const ConcertSyncCore::FReplicationStreamObjectID& Object, TSharedRef<const FConcertObjectReplicationEvent> Data) override
+			{
+				if (EnumHasAnyFlags(Flags, EReplicationCacheTestFlags::NeverReceive))
+				{
+					Test.AddError(TEXT("Received object that we never asked for it!"));
+					return;
+				}
+				
+				if (EnumHasAnyFlags(Flags, EReplicationCacheTestFlags::NeverConsume) && CachedData.IsValid())
+				{
+					Test.AddError(TEXT("Received new object although the old one was not consumed yet!"));
+					return;
+				}
+				
+				if (!EnumHasAnyFlags(Flags, EReplicationCacheTestFlags::InstantConsume))
+				{
+					CachedData = MoveTemp(Data);
+				}
+			}
+
+			float PeakData()
+			{
+				if (!CachedData)
+				{
+					return -1.f;
+				}
+
+				FNativeStruct Data;
+				CachedData->SerializedPayload.GetTypedPayload(Data);
+				return Data.Float;
+			}
+
+			const void* GetDataAddress()
+			{
+				return CachedData.Get();
+			}
+			
+			void Consume()
+			{
+				check(!EnumHasAnyFlags(Flags, EReplicationCacheTestFlags::NeverConsume));
+				CachedData.Reset();
+			}
+		};
+
+
+		// Set up
+		const FGuid StreamId = FGuid::NewGuid();
+		const FSoftObjectPath ObjectPath(TEXT("/Game/World.World:PersistentLevel.StaticMeshActor0"));
+		const ConcertSyncCore::FReplicationStreamObjectID ObjectID{ StreamId, ObjectPath};
+		
+		TSharedRef<ConcertSyncCore::FObjectReplicationCache> Cache = MakeShared<ConcertSyncCore::FObjectReplicationCache>(MakeShared<FTestReplicationFormat>());
+		TSharedRef<FTestReplicationCacheUser> User_NeverConsume = MakeShared<FTestReplicationCacheUser>(*this, ObjectID, EReplicationCacheTestFlags::NeverConsume);
+		TSharedRef<FTestReplicationCacheUser> User_NeverReceive = MakeShared<FTestReplicationCacheUser>(*this, ObjectID, EReplicationCacheTestFlags::NeverReceive);
+		TSharedRef<FTestReplicationCacheUser> User_ConsumeManually = MakeShared<FTestReplicationCacheUser>(*this, ObjectID);
+		TSharedRef<FTestReplicationCacheUser> User_InstantConsume = MakeShared<FTestReplicationCacheUser>(*this, ObjectID, EReplicationCacheTestFlags::InstantConsume);
+		
+		Cache->RegisterDataCacheUser(User_NeverConsume);
+		Cache->RegisterDataCacheUser(User_NeverReceive);
+		Cache->RegisterDataCacheUser(User_ConsumeManually);
+		Cache->RegisterDataCacheUser(User_InstantConsume);
+
+		const FConcertObjectReplicationEvent Event_5 = FTestReplicationFormat::CreateEvent(ObjectPath, 5.f);
+		const FConcertObjectReplicationEvent Event_10 = FTestReplicationFormat::CreateEvent(ObjectPath, 10.f);
+		const FConcertObjectReplicationEvent Event_20 = FTestReplicationFormat::CreateEvent(ObjectPath, 20.f);
+		const FConcertObjectReplicationEvent Event_100 = FTestReplicationFormat::CreateEvent(ObjectPath, 100.f);
+
+
+		
+		// Tests
+		Cache->StoreUntilConsumed(StreamId, Event_5);
+		TestTrue(TEXT("Users have same data"), User_NeverConsume->CachedData.Get() == User_ConsumeManually->CachedData.Get());
+		User_ConsumeManually->Consume();
+		
+		Cache->StoreUntilConsumed(StreamId, Event_10);
+		TestEqual(TEXT("Combined events: 5 and 10"), User_NeverConsume->PeakData(), 15.f);
+		TestEqual(TEXT("Received new data: 10"), User_ConsumeManually->PeakData(), 10.f);
+		
+		// The goal here is to test that the cache does not leak resources when a user de-registers
+		Cache->UnregisterDataCacheUser(User_ConsumeManually);
+		Cache->RegisterDataCacheUser(User_ConsumeManually);
+		const void* AddressBefore = User_ConsumeManually->GetDataAddress();
+		Cache->StoreUntilConsumed(StreamId, Event_100);
+		TestEqual(TEXT("Combined events: 5, 10, and 100"), User_NeverConsume->PeakData(), 115.f);
+		TestEqual(TEXT("Received totally new data"), User_ConsumeManually->PeakData(), 100.f);
+		TestNotEqual(TEXT("Received totally new allocated data"), AddressBefore, User_ConsumeManually->GetDataAddress());
+
+		// Test: No crash or anything
+		Cache->UnregisterDataCacheUser(User_NeverConsume);
+		Cache->UnregisterDataCacheUser(User_NeverReceive);
+		Cache->UnregisterDataCacheUser(User_ConsumeManually);
+		Cache->UnregisterDataCacheUser(User_InstantConsume);
+		return true;
+	}
+}
diff --git a/Engine/Plugins/Developer/Concert/ConcertSync/ConcertSyncTest/Source/ConcertSyncTest/Private/Replication/ReplicationHandshakeTests.cpp b/Engine/Plugins/Developer/Concert/ConcertSync/ConcertSyncTest/Source/ConcertSyncTest/Private/Replication/ReplicationHandshakeTests.cpp
new file mode 100644
--- /dev/null
+++ b/Engine/Plugins/Developer/Concert/ConcertSync/ConcertSyncTest/Source/ConcertSyncTest/Private/Replication/ReplicationHandshakeTests.cpp
@@ -0,0 +1,137 @@
+// Copyright Epic Games, Inc. All Rights Reserved.
+
+#include "Util/ConcertMocks.h"
+
+#include "Replication/Data/ReplicationClientDescription.h"
+#include "Replication/Data/ReplicationStreamDescription.h"
+#include "Replication/IConcertClientReplicationBridge.h"
+#include "Replication/IConcertClientReplicationManager.h"
+#include "Replication/IConcertServerReplicationManager.h"
+#include "Replication/Messages/ConcertReplicationHandshakeMessages.h"
+#include "ReplicationTestInterface.h"
+#include "Util/ConcertClientReplicationBridgeMock.h"
+
+#include "Components/StaticMeshComponent.h"
+#include "Misc/AutomationTest.h"
+#include "Misc/Paths.h"
+#include "HAL/FileManager.h"
+
+namespace UE::ConcertSyncTests
+{
+	/**
+	 * Tests the handshake for joining and leaving a replication session.
+	 */
+	IMPLEMENT_CUSTOM_SIMPLE_AUTOMATION_TEST(FJoinHandshakeTest, FConcertClientServerCommunicationTest, "Concert.Replication.JoinHandshake", EAutomationTestFlags::EditorContext | EAutomationTestFlags::EngineFilter);
+	bool FJoinHandshakeTest::RunTest(const FString& Parameters)
+	{
+		using namespace ConcertSyncClient::TestInterface;
+		using namespace ConcertSyncServer::TestInterface;
+		using namespace ConcertSyncClient::Replication;
+		using namespace ConcertSyncServer::Replication;
+		const FSoftObjectPath PathToSomeActorComponent(TEXT("/Game/Map.Map:PersistentLevel.StaticMeshActor0.StaticMeshComponent0"));
+		const FConcertPropertyChain ForcedLodModelProperty = *FConcertPropertyChain::CreateFromPath(*UStaticMeshComponent::StaticClass(), { TEXT("ForcedLodModel") });
+		const FConcertPropertyChain MinLODProperty = *FConcertPropertyChain::CreateFromPath(*UStaticMeshComponent::StaticClass(), { TEXT("MinLOD") });
+
+		// 1. Init
+		
+		// Server
+		InitServer();
+		const TSharedPtr<IConcertServerSession>& ServerSession = GetServerSessionMock();
+		const TSharedRef<IConcertServerReplicationManager> ServerReplicationManager = CreateServerReplicationManager(ServerSession.ToSharedRef());
+		
+		// Client
+		FClientInfo& Client = ConnectClient();
+		const TSharedRef<IConcertClientReplicationBridge> BridgeMock = MakeShared<FConcertClientReplicationBridgeMock>();
+		const TSharedPtr<IConcertClientSession>& ClientSession = Client.ClientSessionMock;
+		const TSharedRef<IConcertClientReplicationManager> ClientReplicationManager = CreateClientReplicationManager(ClientSession.ToSharedRef(), &BridgeMock.Get());
+		FClientInfo& Client_Secondary = ConnectClient();
+		const TSharedRef<IConcertClientReplicationBridge> BridgeMock_Secondary = MakeShared<FConcertClientReplicationBridgeMock>();
+		const TSharedPtr<IConcertClientSession>& ClientSession_Secondary = Client_Secondary.ClientSessionMock;
+		const TSharedRef<IConcertClientReplicationManager> ClientReplicationManager_Secondary = CreateClientReplicationManager(ClientSession_Secondary.ToSharedRef(), &BridgeMock_Secondary.Get());
+
+		// Stream info
+		FReplicationClientDescription ClientDescription;
+		FReplicationStreamDescription StreamDescription;
+		FReplicatedObjectInfo StaticMeshComponentInfo;
+		StaticMeshComponentInfo.ClassPath = UStaticMeshComponent::StaticClass();
+		StaticMeshComponentInfo.PropertySelection.ReplicatedProperties.Add(ForcedLodModelProperty);
+		StreamDescription.BaseDescription.ReplicationMap.ReplicatedObjects.Add(PathToSomeActorComponent, StaticMeshComponentInfo);
+
+		FReplicationStreamDescription DuplicateStreamDescription = StreamDescription;
+		DuplicateStreamDescription.BaseDescription.Identifier = FGuid::NewGuid();
+		
+		FReplicationStreamDescription InvalidClassStreamDescription = StreamDescription;
+		FReplicatedObjectInfo MissingClassInfo = StaticMeshComponentInfo;
+		MissingClassInfo.ClassPath.Reset();
+		InvalidClassStreamDescription.BaseDescription.ReplicationMap.ReplicatedObjects.Add(PathToSomeActorComponent, MissingClassInfo);
+
+		
+		// 2. Run
+
+		// 2.1 Invalid configurations
+		// 2.1.1 Duplicate properties
+		FReplicationStreamDescription Invalid_DoublePropertyDescription = StreamDescription;
+		Invalid_DoublePropertyDescription.BaseDescription.ReplicationMap.ReplicatedObjects[PathToSomeActorComponent].PropertySelection.ReplicatedProperties.Add(ForcedLodModelProperty);
+		ClientReplicationManager->JoinReplicationSession({ ClientDescription, { Invalid_DoublePropertyDescription } })
+			.Next([&](const FJoinReplicatedSessionResult& Result)
+			{
+				TestTrue(TEXT("Cannot contain same properties twice"), Result.ErrorCode == EJoinReplicationErrorCode::DuplicateProperty);
+			});
+		// 2.1.2 Duplicate stream identifier
+		ClientReplicationManager->JoinReplicationSession({ ClientDescription, { StreamDescription, StreamDescription } })
+			.Next([&](const FJoinReplicatedSessionResult& Result)
+			{
+				TestTrue(TEXT("Cannot contain stream ID twice"), Result.ErrorCode == EJoinReplicationErrorCode::ConflictingStreamId);
+			});
+		// 2.1.3 Miss class path
+		ClientReplicationManager->JoinReplicationSession({ ClientDescription, { InvalidClassStreamDescription } })
+			.Next([&](const FJoinReplicatedSessionResult& Result)
+			{
+				TestTrue(TEXT("Cannot contain null classes"), Result.ErrorCode == EJoinReplicationErrorCode::InvalidClass);
+			});
+
+		// 2.2 Valid join
+		ClientReplicationManager->JoinReplicationSession({ ClientDescription, { StreamDescription } })
+			.Next([&](const FJoinReplicatedSessionResult& Result)
+			{
+				TestTrue(TEXT("Join with valid args"), Result.ErrorCode == EJoinReplicationErrorCode::Success);
+			});
+		// 2.3 No joining twice
+		AddExpectedError(TEXT("JoinReplicationSession requested while already in a session"), EAutomationExpectedErrorFlags::Contains); // Not pretty, but otherwise this test fails due to logged warning
+		ClientReplicationManager->JoinReplicationSession({ ClientDescription, { StreamDescription } })
+			.Next([&](const FJoinReplicatedSessionResult& Result)
+			{
+				TestTrue(TEXT("Cannot join twice"), Result.ErrorCode == EJoinReplicationErrorCode::AlreadyInSession);
+			});
+		// 2.4 Rejoining
+		ClientReplicationManager->LeaveReplicationSession();
+		ClientReplicationManager->JoinReplicationSession({ ClientDescription, { StreamDescription } })
+			.Next([&](const FJoinReplicatedSessionResult& Result)
+			{
+				TestTrue(TEXT("Re-join session"), Result.ErrorCode == EJoinReplicationErrorCode::Success);
+			});
+
+		
+		// 2.5 Second client is not allowed replicate the same properties as the first client
+		FReplicationStreamDescription DuplicateProperties = StreamDescription;
+		DuplicateProperties.BaseDescription.Identifier = FGuid::NewGuid();
+		ClientReplicationManager_Secondary->JoinReplicationSession({ ClientDescription, { DuplicateProperties } })
+			.Next([&](const FJoinReplicatedSessionResult& Result)
+			{
+				TestTrue(TEXT("Clients cannot overlap properties"), Result.ErrorCode == EJoinReplicationErrorCode::ConflictingAuthority);
+			});
+		// 2.6 Second client is allowed to replicate different properties on the same object
+		FReplicationStreamDescription NonOverlappingProperties;
+		FReplicatedObjectInfo StaticMeshComponentInfo_Secondary;
+		StaticMeshComponentInfo_Secondary.ClassPath = UStaticMeshComponent::StaticClass();
+		StaticMeshComponentInfo_Secondary.PropertySelection.ReplicatedProperties.Add(MinLODProperty);
+		NonOverlappingProperties.BaseDescription.ReplicationMap.ReplicatedObjects.Add(PathToSomeActorComponent, StaticMeshComponentInfo_Secondary);
+		ClientReplicationManager_Secondary->JoinReplicationSession({ ClientDescription, { NonOverlappingProperties } })
+		.Next([&](const FJoinReplicatedSessionResult& Result)
+		{
+			TestTrue(TEXT("Two clients can replicate differing properties on the same object"), Result.ErrorCode == EJoinReplicationErrorCode::Success);
+		});
+		
+		return true;
+	}
+}
diff --git a/Engine/Plugins/Developer/Concert/ConcertSync/ConcertSyncTest/Source/ConcertSyncTest/Private/Replication/ReplicationSendReceiveTests.cpp b/Engine/Plugins/Developer/Concert/ConcertSync/ConcertSyncTest/Source/ConcertSyncTest/Private/Replication/ReplicationSendReceiveTests.cpp
new file mode 100644
--- /dev/null
+++ b/Engine/Plugins/Developer/Concert/ConcertSync/ConcertSyncTest/Source/ConcertSyncTest/Private/Replication/ReplicationSendReceiveTests.cpp
@@ -0,0 +1,201 @@
+// Copyright Epic Games, Inc. All Rights Reserved.
+
+#include "Replication/Data/ReplicationStreamDescription.h"
+#include "Replication/IConcertClientReplicationManager.h"
+#include "Replication/IConcertServerReplicationManager.h"
+#include "Replication/Messages/ConcertReplicationEvents.h"
+#include "Replication/Messages/ConcertReplicationHandshakeMessages.h"
+#include "Replication/PropertyChainUtils.h"
+#include "Replication/ReplicationTestInterface.h"
+#include "TestReflectionObject.h"
+#include "Util/ConcertClientReplicationBridgeMock.h"
+#include "Util/ConcertMocks.h"
+
+#include "Misc/AutomationTest.h"
+#include "Misc/Paths.h"
+#include "HAL/FileManager.h"
+#include "UObject/UObjectGlobals.h"
+#include "UObject/Package.h"
+
+namespace UE::ConcertSyncTests
+{
+	static void SetTestValues(UTestReflectionObject& Object);
+	static void SetDifferentValues(UTestReflectionObject& Object);
+	static void TestEqualTestValues(UTestReflectionObject& Object, FAutomationTestBase& Test);
+
+	static void TickClient(FConcertClientServerCommunicationTest::FClientInfo& Client);
+	static void TickServer(FConcertServerSessionMock& ServerSession);
+
+	/**
+	 * Tests replicating data from sender client > server > receiver client.
+	 */
+	IMPLEMENT_CUSTOM_SIMPLE_AUTOMATION_TEST(FSendReceiveLogicTests, FConcertClientServerCommunicationTest, "Concert.Replication.SendReceiveLogic", EAutomationTestFlags::EditorContext | EAutomationTestFlags::EngineFilter);
+	bool FSendReceiveLogicTests::RunTest(const FString& Parameters)
+	{
+		using namespace ConcertSyncClient::TestInterface;
+		using namespace ConcertSyncServer::TestInterface;
+		using namespace ConcertSyncClient::Replication;
+		using namespace ConcertSyncServer::Replication;
+
+		// 1. Init
+		
+		// Server
+		InitServer();
+		const TSharedPtr<FConcertServerSessionMock>& ServerSession = GetServerSessionMock();
+		const TSharedRef<IConcertServerReplicationManager> ServerReplicationManager = CreateServerReplicationManager(ServerSession.ToSharedRef());
+		
+		// Client Receiver
+		FClientInfo& Client_Receiver = ConnectClient();
+		const TSharedRef<FConcertClientReplicationBridgeMock> BridgeMock_Receiver = MakeShared<FConcertClientReplicationBridgeMock>();
+		const TSharedPtr<IConcertClientSession>& ClientSession_Receiver = Client_Receiver.ClientSessionMock;
+		const TSharedRef<IConcertClientReplicationManager> ClientReplicationManager_Receiver = CreateClientReplicationManager(ClientSession_Receiver.ToSharedRef(), &BridgeMock_Receiver.Get());
+		
+		// Client Sender
+		FClientInfo& Client_Sender = ConnectClient();
+		const TSharedRef<FConcertClientReplicationBridgeMock> BridgeMock_Sender = MakeShared<FConcertClientReplicationBridgeMock>();
+		const TSharedPtr<IConcertClientSession>& ClientSession_Sender = Client_Sender.ClientSessionMock;
+		const TSharedRef<IConcertClientReplicationManager> ClientReplicationManager_Sender = CreateClientReplicationManager(ClientSession_Sender.ToSharedRef(), &BridgeMock_Sender.Get());
+
+		// Fake replicated object
+		UTestReflectionObject* TestObject = NewObject<UTestReflectionObject>(GetTransientPackage());
+		const FSoftObjectPath TestObjectPath = TestObject;
+		
+		// 1.1 Sender offers to send all UTestReflectionObject properties
+		const FGuid SenderStreamId = FGuid::NewGuid();
+		{
+			FReplicatedObjectInfo AllProperties { UTestReflectionObject::StaticClass() };
+			ConcertSyncCore::PropertyChain::ForEachReplicatableConcertProperty(*UTestReflectionObject::StaticClass(), [&AllProperties](FConcertPropertyChain&& Chain)
+			{
+				AllProperties.PropertySelection.ReplicatedProperties.Emplace(MoveTemp(Chain));
+				return EBreakBehavior::Continue;
+			});
+			
+			FJoinReplicatedSessionArgs SenderJoinArgs;
+			FReplicationStreamDescription SendingStream;
+			SendingStream.BaseDescription.Identifier = SenderStreamId;
+			SendingStream.BaseDescription.ReplicationMap.ReplicatedObjects.Add(TestObjectPath, AllProperties);
+			SenderJoinArgs.Streams.Add(SendingStream);
+			ClientReplicationManager_Sender->JoinReplicationSession(SenderJoinArgs)
+				.Next([this](const FJoinReplicatedSessionResult& Result){ TestTrue(TEXT("Sender joined"), Result.ErrorCode == EJoinReplicationErrorCode::Success); });
+		}
+
+		// 1.2 Receiver accepts all properties
+		{
+			FJoinReplicatedSessionArgs ReceiverJoinArgs;
+			// TODO: When we add client attributes, this test must be updated with attributes that allow receiving all data
+			ClientReplicationManager_Receiver->JoinReplicationSession(ReceiverJoinArgs)
+				.Next([this](const FJoinReplicatedSessionResult& Result){ TestTrue(TEXT("Receiver joined"), Result.ErrorCode == EJoinReplicationErrorCode::Success); });
+		}
+		
+		// 1.3 Prepare test data
+		bool bHasServerReceivedData = false;
+		bool bHasClientReceivedData = false;
+		auto TestReplicationData_Server = [&](const FConcertSessionContext& Context, const FConcertBatchReplicationEvent& Event)
+		{
+			if (bHasServerReceivedData)
+			{
+				AddError(TEXT("Server was expected to receive data exactly once!"));
+				return;
+			}
+			bHasServerReceivedData = true;
+			
+			TestEqual(TEXT("Server received 1 stream"), Event.Streams.Num(), 1);
+			TestEqual(TEXT("Server received 1 object"), Event.Streams.IsEmpty() ? 0 : Event.Streams[0].ReplicatedObjects.Num() , 1);
+			TestEqual(TEXT("Server received from correct stream"), Event.Streams.IsEmpty() ? FGuid{} : Event.Streams[0].StreamId, SenderStreamId);
+			const FSoftObjectPath ObjectPath = Event.Streams.IsEmpty() || Event.Streams[0].ReplicatedObjects.IsEmpty() ? FSoftObjectPath{} : Event.Streams[0].ReplicatedObjects[0].ReplicatedObject;
+			TestEqual(TEXT("Server's received object has correct path"), ObjectPath, TestObjectPath);
+		};
+		auto TestReplicationData_Client_Sender = [&](const FConcertSessionContext& Context, const FConcertBatchReplicationEvent& Event)
+        {
+        	// There was a bug where client would receive its own data...
+			AddError(TEXT("Client 1 was not supposed to receive data!"));
+        };
+		auto TestReplicationData_Client_Receiver = [&](const FConcertSessionContext& Context, const FConcertBatchReplicationEvent& Event)
+		{
+			if (bHasClientReceivedData)
+			{
+				AddError(TEXT("Client 2 was expected to receive data exactly once!"));
+				return;
+			}
+			bHasClientReceivedData = true;
+			
+			TestEqual(TEXT("Client 2 received 1 stream"), Event.Streams.Num(), 1);
+			TestEqual(TEXT("Client 2 received 1 object"), Event.Streams.IsEmpty() ? 0 : Event.Streams[0].ReplicatedObjects.Num() , 1);
+			TestEqual(TEXT("Client 2 received from correct stream"), Event.Streams.IsEmpty() ? FGuid{} : Event.Streams[0].StreamId, SenderStreamId);
+			const FSoftObjectPath ObjectPath = Event.Streams.IsEmpty() || Event.Streams[0].ReplicatedObjects.IsEmpty() ? FSoftObjectPath{} : Event.Streams[0].ReplicatedObjects[0].ReplicatedObject;
+			TestEqual(TEXT("Client 2's received object has correct path"), ObjectPath, TestObjectPath);
+		};
+		ServerSession->RegisterCustomEventHandler<FConcertBatchReplicationEvent>(TestReplicationData_Server);
+		Client_Sender.ClientSessionMock->RegisterCustomEventHandler<FConcertBatchReplicationEvent>(TestReplicationData_Client_Sender);
+		Client_Receiver.ClientSessionMock->RegisterCustomEventHandler<FConcertBatchReplicationEvent>(TestReplicationData_Client_Receiver);
+
+		// 1.4 Replication bridge
+		// Bridge is responsible for telling client-side replication system about existing objects
+		BridgeMock_Sender->InjectAvailableObject(*TestObject); 
+		BridgeMock_Receiver->InjectAvailableObject(*TestObject);
+
+		
+
+		
+		// 2. Send data
+		// Note that the systems operate on the same UObject instance. This is required since the systems internally use object path.
+		// For this reason we change the object after sending.
+		
+		// 2.1 Send from sender to server
+		SetTestValues(*TestObject);
+		TickClient(Client_Sender);
+		
+		// 2.2 Forward from server to receiver
+		TickServer(*ServerSession);
+		
+		// 2.3 Receive from server
+		SetDifferentValues(*TestObject);
+		TickClient(Client_Receiver);
+
+
+
+
+		
+		// 3. Test
+		// ... that the data went through server and receiving client ...
+		TestTrue(TEXT("Server received replication event"), bHasServerReceivedData);
+		TestTrue(TEXT("Client 2 received replication event"), bHasClientReceivedData);
+		// ... and the data was applied correctly to the object
+		TestEqualTestValues(*TestObject, *this);
+		return true;
+	}
+
+	// TODO: Add more properties here
+	constexpr float SentFloat = 42.f;
+	const FVector SentVector = { 21.f, 84.f, -1.f };
+
+	static void SetTestValues(UTestReflectionObject& Object)
+	{
+		Object.Float = SentFloat;
+		Object.Vector = SentVector;
+	}
+	
+	static void SetDifferentValues(UTestReflectionObject& Object)
+	{
+		Object.Float *= -1.f;
+		Object.Vector *= -1.f;
+	}
+	
+	static void TestEqualTestValues(UTestReflectionObject& Object, FAutomationTestBase& Test)
+	{
+		Test.TestEqual(TEXT("Float"), Object.Float, SentFloat);
+		Test.TestEqual(TEXT("Vector"), Object.Vector, SentVector);
+	}
+	
+	constexpr float FakeDeltaTime = 1.f / 60.f;
+	
+	static void TickClient(FConcertClientServerCommunicationTest::FClientInfo& Client)
+	{
+		Client.ClientSessionMock->OnTick().Broadcast(*Client.ClientSessionMock, FakeDeltaTime);
+	}
+
+	static void TickServer(FConcertServerSessionMock& ServerSession)
+	{
+		ServerSession.OnTick().Broadcast(ServerSession, FakeDeltaTime);
+	}
+}
diff --git a/Engine/Plugins/Developer/Concert/ConcertSync/ConcertSyncTest/Source/ConcertSyncTest/Private/Replication/ReplicationTestInterface.h b/Engine/Plugins/Developer/Concert/ConcertSync/ConcertSyncTest/Source/ConcertSyncTest/Private/Replication/ReplicationTestInterface.h
new file mode 100644
--- /dev/null
+++ b/Engine/Plugins/Developer/Concert/ConcertSync/ConcertSyncTest/Source/ConcertSyncTest/Private/Replication/ReplicationTestInterface.h
@@ -0,0 +1,28 @@
+// Copyright Epic Games, Inc. All Rights Reserved.
+
+#pragma once
+
+class IConcertClientSession;
+class IConcertServerSession;
+class IConcertClientReplicationBridge;
+class IConcertClientReplicationManager;
+
+namespace UE::ConcertSyncServer::Replication
+{
+	class IConcertServerReplicationManager;
+}
+
+namespace UE::ConcertSyncClient::TestInterface
+{
+	extern CONCERTSYNCCLIENT_API TSharedRef<IConcertClientReplicationManager> CreateClientReplicationManager(
+		TSharedRef<IConcertClientSession> InLiveSession,
+		IConcertClientReplicationBridge* InBridge
+		);
+}
+
+namespace UE::ConcertSyncServer::TestInterface
+{
+	extern CONCERTSYNCSERVER_API TSharedRef<Replication::IConcertServerReplicationManager> CreateServerReplicationManager(
+		TSharedRef<IConcertServerSession> InLiveSession
+		);
+}
\ No newline at end of file
diff --git a/Engine/Plugins/Developer/Concert/ConcertSync/ConcertSyncTest/Source/ConcertSyncTest/Private/Replication/TestReflectionObject.h b/Engine/Plugins/Developer/Concert/ConcertSync/ConcertSyncTest/Source/ConcertSyncTest/Private/Replication/TestReflectionObject.h
--- a/Engine/Plugins/Developer/Concert/ConcertSync/ConcertSyncTest/Source/ConcertSyncTest/Private/Replication/TestReflectionObject.h
+++ b/Engine/Plugins/Developer/Concert/ConcertSync/ConcertSyncTest/Source/ConcertSyncTest/Private/Replication/TestReflectionObject.h
@@ -38,7 +38,7 @@ struct FTestReplicationStruct
 	GENERATED_BODY()
 
 	UPROPERTY()
-	float Value{};
+	float Value{}; // The name of this property is purposefully equal to FConcertPropertyChain::InternalContainerPropertyValueName.
 	UPROPERTY()
 	FVector Vector = FVector::ZeroVector;
 	UPROPERTY()
diff --git a/Engine/Plugins/Developer/Concert/ConcertSync/ConcertSyncTest/Source/ConcertSyncTest/Private/Util/ConcertClientReplicationBridgeMock.h b/Engine/Plugins/Developer/Concert/ConcertSync/ConcertSyncTest/Source/ConcertSyncTest/Private/Util/ConcertClientReplicationBridgeMock.h
new file mode 100644
--- /dev/null
+++ b/Engine/Plugins/Developer/Concert/ConcertSync/ConcertSyncTest/Source/ConcertSyncTest/Private/Util/ConcertClientReplicationBridgeMock.h
@@ -0,0 +1,58 @@
+// Copyright Epic Games, Inc. All Rights Reserved.
+
+#pragma once
+
+#include "Algo/AnyOf.h"
+#include "Replication/IConcertClientReplicationBridge.h"
+
+namespace UE::ConcertSyncTests
+{
+	class FConcertClientReplicationBridgeMock : public IConcertClientReplicationBridge
+	{
+	public:
+
+		TArray<FSoftObjectPath> TrackedObjects;
+		TArray<UObject*> AvailableObjects;
+		
+		FConcertClientReplicationBridgeObjectEvent OnObjectDiscoveredDelegate;
+		FConcertClientReplicationBridgeObjectEvent OnObjectRemovedDelegate;
+
+		void InjectAvailableObject(UObject& Object)
+		{
+			AvailableObjects.Add(&Object);
+			OnObjectDiscoveredDelegate.Broadcast(Object);
+		}
+		
+		virtual void PushTrackedObjects(TArrayView<const FSoftObjectPath> InTrackedObjects) override
+		{
+			for (const FSoftObjectPath& Path : InTrackedObjects)
+			{
+				TrackedObjects.Add(Path);
+			}
+		}
+		
+		virtual void PopTrackedObjects(TArrayView<const FSoftObjectPath> InTrackedObjects) override
+		{
+			for (const FSoftObjectPath& Path : InTrackedObjects)
+			{
+				TrackedObjects.RemoveSingle(Path);
+			}
+		}
+
+		virtual bool IsObjectAvailable(const FSoftObjectPath& Path) override { return Algo::AnyOf(AvailableObjects, [&Path](UObject* Object){ return Path == Object; }); }
+		virtual UObject* FindObjectIfAvailable(const FSoftObjectPath& Path) override
+		{
+			for (UObject* Object : AvailableObjects)
+			{
+				if (Path == Object)
+				{
+					return Object;
+				}
+			}
+			return nullptr;
+		}
+
+		virtual FConcertClientReplicationBridgeObjectEvent& OnObjectDiscovered() override { return OnObjectDiscoveredDelegate; }
+		virtual FConcertClientReplicationBridgeObjectEvent& OnObjectHidden() override { return OnObjectRemovedDelegate; }
+	};
+}
diff --git a/Engine/Plugins/Developer/Concert/ConcertSync/ConcertSyncTest/Source/ConcertSyncTest/Private/Util/ConcertMocks.h b/Engine/Plugins/Developer/Concert/ConcertSync/ConcertSyncTest/Source/ConcertSyncTest/Private/Util/ConcertMocks.h
new file mode 100644
--- /dev/null
+++ b/Engine/Plugins/Developer/Concert/ConcertSync/ConcertSyncTest/Source/ConcertSyncTest/Private/Util/ConcertMocks.h
@@ -0,0 +1,435 @@
+// Copyright Epic Games, Inc. All Rights Reserved.
+
+#pragma once
+
+#include "ConcertMessages.h"
+#include "IConcertSession.h"
+#include "Misc/AutomationTest.h"
+#include "Misc/Paths.h"
+#include "Scratchpad/ConcertScratchpad.h"
+#include "Scratchpad/ConcertScratchpadPtr.h"
+#include "UObject/StructOnScope.h"
+
+namespace UE::ConcertSyncTests
+{
+	// Utility functions used to detect when a non-mocked function is called, so that we can mock it properly when required.
+	template<typename T> T NotMocked(T Ret) { check(false); return Ret; }
+	template<typename T> T NotMocked()      { check(false); return T(); }
+
+	inline FString GetTestSessionRootPath()
+	{
+		return FPaths::ProjectIntermediateDir() / TEXT("ConcertDataStoreTest");
+	}
+	
+	/** Implements a not-working IConcertServerSession. It must be further overridden to implement just what is required by the tests */
+	class FConcertServerSessionBaseMock : public IConcertServerSession
+	{
+	public:
+		FConcertServerSessionBaseMock() 
+			: Id(FGuid::NewGuid())
+			, Name("FConcertServerSessionBaseMock")
+		{ }
+
+		// IConcertSession Begin.
+		virtual const FGuid& GetId() const override																								{ return Id; }
+		virtual const FString& GetName() const override																							{ return NotMocked<const FString&>(Name); }
+		virtual void SetName(const FString&) override																							{ return NotMocked<void>(); }
+		virtual FMessageAddress GetClientAddress(const FGuid& ClientEndpointId) const override													{ return NotMocked<FMessageAddress>(); }
+		virtual const FConcertSessionInfo& GetSessionInfo() const override																		{ return NotMocked<const FConcertSessionInfo&>(SessionInfo); }
+		virtual FString GetSessionWorkingDirectory() const override																				{ return GetTestSessionRootPath() / Id.ToString(); }
+		virtual TArray<FGuid> GetSessionClientEndpointIds() const override																		{ return NotMocked<TArray<FGuid>>(); }
+		virtual TArray<FConcertSessionClientInfo> GetSessionClients() const override															{ return NotMocked<TArray<FConcertSessionClientInfo>>(); }
+		virtual bool FindSessionClient(const FGuid&, FConcertSessionClientInfo&) const override													{ return NotMocked<bool>(); }
+		virtual void Startup() override																											{ return NotMocked<void>(); }
+		virtual void Shutdown() override																										{ return NotMocked<void>(); };
+		virtual FConcertScratchpadRef GetScratchpad() const override																			{ return NotMocked<FConcertScratchpadRef>(); }
+		virtual FConcertScratchpadPtr GetClientScratchpad(const FGuid&) const override															{ return NotMocked<FConcertScratchpadRef>(); }
+		virtual FDelegateHandle InternalRegisterCustomEventHandler(const FName&, const TSharedRef<IConcertSessionCustomEventHandler>&) override { return NotMocked<FDelegateHandle>(); }
+		virtual void InternalUnregisterCustomEventHandler(const FName& EventMessageType, const FDelegateHandle EventHandle) override			{ return NotMocked<void>(); }
+		virtual void InternalUnregisterCustomEventHandler(const FName& EventMessageType, const void* EventHandler) override						{ return NotMocked<void>(); }
+		virtual void InternalClearCustomEventHandler(const FName& EventMessageType) override													{ return NotMocked<void>(); }
+		virtual void InternalSendCustomEvent(const UScriptStruct*, const void*, const TArray<FGuid>&, EConcertMessageFlags, TOptional<FConcertSequencedCustomEvent> InSequenceId={}) override			{ return NotMocked<void>(); }
+		virtual void InternalRegisterCustomRequestHandler(const FName&, const TSharedRef<IConcertSessionCustomRequestHandler>&) override		{ return NotMocked<void>(); }
+		virtual void InternalUnregisterCustomRequestHandler(const FName&) override																{ return NotMocked<void>(); }
+		virtual void InternalSendCustomRequest(const UScriptStruct*, const void*, const FGuid&, const TSharedRef<IConcertSessionCustomResponseHandler>&) override { NotMocked<void>(); }
+		// IConcertSession End.
+
+		// IConcertServerSession Begin
+		virtual FOnConcertServerSessionTick& OnTick() override                          { return NotMocked<FOnConcertServerSessionTick&>(Tick); }
+		virtual FOnConcertServerSessionClientChanged& OnSessionClientChanged() override { return NotMocked<FOnConcertServerSessionClientChanged&>(ConnectionChanged); }
+		virtual FOnConcertMessageAcknowledgementReceivedFromLocalEndpoint& OnConcertMessageAcknowledgementReceived() override { return NotMocked<FOnConcertMessageAcknowledgementReceivedFromLocalEndpoint&>(AckReceived); }
+		// IConcertServerSession End
+
+	protected:
+		FGuid Id;
+
+		// Those below are unused mocked data member, but required to get the code compiling.
+		FString Name;
+		FConcertSessionInfo SessionInfo;
+		FOnConcertServerSessionTick Tick;
+		FOnConcertServerSessionClientChanged ConnectionChanged;
+		FOnConcertMessageAcknowledgementReceivedFromLocalEndpoint AckReceived;
+	};
+
+	/** Implements a not-working IConcertClientSession. It must be further overridden to implement just what is required by the tests */
+	class FConcertClientSessionBaseMock : public IConcertClientSession
+	{
+	public:
+		FConcertClientSessionBaseMock() 
+			: Id(FGuid::NewGuid()) 
+			, Name("FConcertClientSessionBaseMock") 
+		{ }
+
+		// IConcertSession Begin.
+		virtual const FGuid& GetId() const override																								{ return Id; }
+		virtual const FString& GetName() const override																							{ return NotMocked<const FString&>(Name); }
+		virtual const FConcertSessionInfo& GetSessionInfo() const override																		{ return NotMocked<const FConcertSessionInfo&>(SessionInfo); }
+		virtual FString GetSessionWorkingDirectory() const override																				{ return GetTestSessionRootPath() / Id.ToString(); }
+		virtual TArray<FGuid> GetSessionClientEndpointIds() const override																		{ return NotMocked<TArray<FGuid>>(); }
+		virtual TArray<FConcertSessionClientInfo> GetSessionClients() const override															{ return NotMocked<TArray<FConcertSessionClientInfo>>(); }
+		virtual bool FindSessionClient(const FGuid&, FConcertSessionClientInfo&) const override													{ return NotMocked<bool>(); }
+		virtual void Startup() override																											{ return NotMocked<void>(); }
+		virtual void Shutdown() override																										{ return NotMocked<void>(); };
+		virtual FConcertScratchpadRef GetScratchpad() const override																			{ return NotMocked<FConcertScratchpadRef>(); }
+		virtual FConcertScratchpadPtr GetClientScratchpad(const FGuid&) const override															{ return NotMocked<FConcertScratchpadRef>(); }
+		virtual FDelegateHandle InternalRegisterCustomEventHandler(const FName&, const TSharedRef<IConcertSessionCustomEventHandler>&) override { return NotMocked<FDelegateHandle>(); }
+		virtual void InternalUnregisterCustomEventHandler(const FName& EventMessageType, const FDelegateHandle EventHandle) override			{ return NotMocked<void>(); }
+		virtual void InternalUnregisterCustomEventHandler(const FName& EventMessageType, const void* EventHandler) override						{ return NotMocked<void>(); }
+		virtual void InternalClearCustomEventHandler(const FName& EventMessageType) override													{ return NotMocked<void>(); }
+		virtual void InternalSendCustomEvent(const UScriptStruct*, const void*, const TArray<FGuid>&, EConcertMessageFlags, TOptional<FConcertSequencedCustomEvent> InSequenceId={}) override			{ return NotMocked<void>(); }
+		virtual void InternalRegisterCustomRequestHandler(const FName&, const TSharedRef<IConcertSessionCustomRequestHandler>&) override		{ return NotMocked<void>(); }
+		virtual void InternalUnregisterCustomRequestHandler(const FName&) override																{ return NotMocked<void>(); }
+		virtual void InternalSendCustomRequest(const UScriptStruct*, const void*, const FGuid&, const TSharedRef<IConcertSessionCustomResponseHandler>&) override { NotMocked<void>(); }
+		// IConcertSession End.
+
+		virtual FConcertSequencedCustomEventManager& GetSequencedEventManager() override
+		{
+			return NotMocked<FConcertSequencedCustomEventManager&>(CustomEventSequenceManager);
+		}
+		// IConcertClientSession Begin
+		virtual EConcertConnectionStatus GetConnectionStatus() const override            { return NotMocked<EConcertConnectionStatus>(EConcertConnectionStatus::Connected); }
+		virtual FGuid GetSessionClientEndpointId() const override                        { return NotMocked<FGuid>(); }
+		virtual FGuid GetSessionServerEndpointId() const override                        { return NotMocked<FGuid>(); }
+		virtual const FConcertClientInfo& GetLocalClientInfo() const override            { return NotMocked<const FConcertClientInfo&>(ClientInfo); }
+		virtual void UpdateLocalClientInfo(const FConcertClientInfoUpdate&) override     { return NotMocked<void>(); }
+		virtual void Connect() override                                                  { return NotMocked<void>(); }
+		virtual void Disconnect() override                                               { return NotMocked<void>(); }
+
+		virtual EConcertSendReceiveState GetSendReceiveState() const override
+		{
+			return NotMocked<EConcertSendReceiveState>(EConcertSendReceiveState::Default);
+		}
+
+		virtual void SetSendReceiveState(EConcertSendReceiveState InSendReceiveState) override
+		{
+			return NotMocked<void>();
+		}
+
+		virtual FOnConcertClientSessionTick& OnTick() override                           { return NotMocked<FOnConcertClientSessionTick&>(Tick); }
+		virtual FOnConcertClientSessionConnectionChanged& OnConnectionChanged() override { return NotMocked<FOnConcertClientSessionConnectionChanged&>(ConnectionChanged); }
+		virtual FOnConcertClientSessionClientChanged& OnSessionClientChanged() override  { return NotMocked<FOnConcertClientSessionClientChanged&>(ClientChanged); }
+		virtual FOnConcertSessionRenamed& OnSessionRenamed() override                    { return NotMocked<FOnConcertSessionRenamed&>(SessionRenamed); }
+		// IConcertClientSession End
+
+		virtual void HandleCustomEvent(const UScriptStruct* EventType, const void* EventData) = 0;
+
+	protected:
+		FGuid Id;
+
+		// Those below are unused mocked data member, but required to get the code compiling.
+		FString Name;
+		FConcertSessionInfo SessionInfo;
+		FOnConcertClientSessionTick Tick;
+		FOnConcertClientSessionConnectionChanged ConnectionChanged;
+		FOnConcertClientSessionClientChanged ClientChanged;
+		FOnConcertSessionRenamed SessionRenamed;
+		FConcertSequencedCustomEventManager CustomEventSequenceManager;
+		FConcertClientInfo ClientInfo;
+	};
+
+	/** Specializes the base concert server session to act as a fake server session. */
+	class FConcertServerSessionMock : public FConcertServerSessionBaseMock
+	{
+	public:
+		
+		virtual void InternalSendCustomEvent(const UScriptStruct* EventType, const void* EventData, const TArray<FGuid>& TargetEndpointIds, EConcertMessageFlags, TOptional<FConcertSequencedCustomEvent> InSequenceId={}) override
+		{
+			for (const FGuid& TargetEndPointId : TargetEndpointIds)
+			{
+				int i = 0;
+				for (const FGuid& ClientEndPointId : ClientEndpoints)
+				{
+					if (TargetEndPointId == ClientEndPointId)
+					{
+						// Dispatch the event on the client immediately.
+						ClientSessions[i]->HandleCustomEvent(EventType, EventData);
+					}
+					++i;
+				}
+			}
+		}
+
+		virtual FDelegateHandle InternalRegisterCustomEventHandler(const FName& EventMessageType, const TSharedRef<IConcertSessionCustomEventHandler>& Handler) override
+		{
+			CustomEventHandlers.FindOrAdd(EventMessageType).Add(Handler);
+			return Handler->GetHandle();
+		}
+
+		virtual void InternalUnregisterCustomEventHandler(const FName& EventMessageType, const FDelegateHandle EventHandle) override
+		{
+			TArray<TSharedPtr<IConcertSessionCustomEventHandler>>* HandlerArrayPtr = CustomEventHandlers.Find(EventMessageType);
+			if (HandlerArrayPtr)
+			{
+				for (auto It = HandlerArrayPtr->CreateIterator(); It; ++It)
+				{
+					if ((*It)->GetHandle() == EventHandle)
+					{
+						It.RemoveCurrent();
+						break;
+					}
+				}
+			}
+		}
+
+		virtual void InternalUnregisterCustomEventHandler(const FName& EventMessageType, const void* EventHandler) override
+		{
+			TArray<TSharedPtr<IConcertSessionCustomEventHandler>>* HandlerArrayPtr = CustomEventHandlers.Find(EventMessageType);
+			if (HandlerArrayPtr)
+			{
+				for (auto It = HandlerArrayPtr->CreateIterator(); It; ++It)
+				{
+					if ((*It)->HasSameObject(EventHandler))
+					{
+						It.RemoveCurrent();
+						break;
+					}
+				}
+			}
+		}
+
+		virtual void InternalRegisterCustomRequestHandler(const FName& RequestMessageType, const TSharedRef<IConcertSessionCustomRequestHandler>& Handler) override
+		{
+			CustomRequestHandlers.Add(RequestMessageType, Handler);
+		}
+
+		virtual void InternalUnregisterCustomRequestHandler(const FName& RequestMessageType) override
+		{
+			CustomRequestHandlers.Remove(RequestMessageType);
+		}
+
+		virtual TArray<FGuid> GetSessionClientEndpointIds() const override
+		{
+			return ClientEndpoints;
+		}
+
+		// Called by the tests to connect a client to the server.
+		void ConnectClient(const FGuid& ClientEndpointId, FConcertClientSessionBaseMock& ClientSession)
+		{
+			// Notify the server that a new client connected. The server data store will replicate its content on the client.
+			FConcertSessionClientInfo Info;
+			Info.ClientEndpointId = ClientEndpointId;
+			ClientEndpoints.Add(ClientEndpointId);
+			ClientSessions.Add(&ClientSession);
+			ConnectionChanged.Broadcast(*this, EConcertClientStatus::Connected, Info);
+		}
+
+		void DispatchEvent(const FGuid& RequesterEndpointId, const UScriptStruct* EventType, const void* EventData, const TArray<FGuid>& DestinationEndpointIds, EConcertMessageFlags Flags, TOptional<FConcertSequencedCustomEvent> InSequenceId)
+		{
+			// { 0, 0, 0, 0 } is for the server
+			if (DestinationEndpointIds.Contains(FGuid{ 0, 0, 0, 0 }))
+			{
+				if (const TArray<TSharedPtr<IConcertSessionCustomEventHandler>>* Handlers = CustomEventHandlers.Find(EventType->GetFName()))
+				{
+					FConcertSessionContext DummyContext;
+					DummyContext.SourceEndpointId = RequesterEndpointId;
+					for (const TSharedPtr<IConcertSessionCustomEventHandler>& Handler : *Handlers)
+					{
+						Handler->HandleEvent(DummyContext, EventData);
+					}
+				}
+			}
+			
+			InternalSendCustomEvent(EventType, EventData, DestinationEndpointIds, Flags, InSequenceId);
+		}
+
+		// Called by the FConcertClientSessionMock to dispatch a request
+		void DispatchRequest(const FGuid& RequesterEndpointId, const UScriptStruct* RequestType, const void* RequestData, const TSharedRef<IConcertSessionCustomResponseHandler>& ResponseHandler)
+		{
+			if (TSharedPtr<IConcertSessionCustomRequestHandler>* RequestHandler = CustomRequestHandlers.Find(RequestType->GetFName()))
+			{
+				// Set up who's sending the request.
+				FConcertSessionContext Context;
+				Context.SourceEndpointId = RequesterEndpointId;
+
+				// Dispatch the request
+				FStructOnScope ResponsePayload((*RequestHandler)->GetResponseType());
+				EConcertSessionResponseCode Result = (*RequestHandler)->HandleRequest(Context, RequestData, ResponsePayload.GetStructMemory());
+				if (Result == EConcertSessionResponseCode::Success || Result == EConcertSessionResponseCode::Failed)
+				{
+					// Dispatch the response.
+					ResponseHandler->HandleResponse(ResponsePayload.GetStructMemory());
+				}
+				else
+				{
+					check(false); // The test suite is not expected to fire any other result than Success or Failed.
+				}
+			}
+		}
+
+		virtual FOnConcertServerSessionClientChanged& OnSessionClientChanged() override
+		{
+			return ConnectionChanged;
+		}
+
+		virtual FOnConcertServerSessionTick& OnTick() override { return OnTickDelegate; }
+
+	private:
+
+		/** Tick delegate for tests that need it */
+		FOnConcertServerSessionTick OnTickDelegate;
+		
+		/** Map of session custom event handlers */
+		TMap<FName, TArray<TSharedPtr<IConcertSessionCustomEventHandler>>> CustomEventHandlers;
+		/** Map of session custom request handlers */
+		TMap<FName, TSharedPtr<IConcertSessionCustomRequestHandler>> CustomRequestHandlers;
+
+		/** Connected client endpoints. */
+		TArray<FGuid> ClientEndpoints;
+
+		/** Connected clients sessions. */
+		TArray<FConcertClientSessionBaseMock*> ClientSessions;
+	};
+
+	/** Specializes the base concert client session to act as a fake client session. */
+	class FConcertClientSessionMock : public FConcertClientSessionBaseMock
+	{
+	public:
+		FConcertClientSessionMock(const FGuid& ClientEndpointId, FConcertServerSessionMock& Server)
+			: ServerMock(Server)
+			, EndpointId(ClientEndpointId)
+		{
+		}
+
+		virtual FDelegateHandle InternalRegisterCustomEventHandler(const FName& EventMessageType, const TSharedRef<IConcertSessionCustomEventHandler>& Handler) override
+		{
+			CustomEventHandlers.FindOrAdd(EventMessageType).Add(Handler);
+			return Handler->GetHandle();
+		}
+
+		virtual void InternalUnregisterCustomEventHandler(const FName& EventMessageType, const FDelegateHandle EventHandle) override
+		{
+			TArray<TSharedPtr<IConcertSessionCustomEventHandler>>* HandlerArrayPtr = CustomEventHandlers.Find(EventMessageType);
+			if (HandlerArrayPtr)
+			{
+				for (auto It = HandlerArrayPtr->CreateIterator(); It; ++It)
+				{
+					if ((*It)->GetHandle() == EventHandle)
+					{
+						It.RemoveCurrent();
+						break;
+					}
+				}
+			}
+		}
+
+		virtual void InternalUnregisterCustomEventHandler(const FName& EventMessageType, const void* EventHandler) override
+		{
+			TArray<TSharedPtr<IConcertSessionCustomEventHandler>>* HandlerArrayPtr = CustomEventHandlers.Find(EventMessageType);
+			if (HandlerArrayPtr)
+			{
+				HandlerArrayPtr->RemoveAllSwap([EventHandler](const TSharedPtr<IConcertSessionCustomEventHandler>& InHandler) -> bool
+				{
+					return InHandler->HasSameObject(EventHandler);
+				});
+			}
+		}
+
+		virtual void InternalClearCustomEventHandler(const FName& EventMessageType) override
+		{
+			CustomEventHandlers.Remove(EventMessageType);
+		}
+
+		virtual void InternalSendCustomEvent(const UScriptStruct* EventType, const void* EventData, const TArray<FGuid>& DestinationEndpointIds, EConcertMessageFlags Flags, TOptional<FConcertSequencedCustomEvent> InSequenceId) override
+		{
+			ServerMock.DispatchEvent(EndpointId, EventType, EventData, DestinationEndpointIds, Flags, InSequenceId);
+		}
+
+		virtual void InternalSendCustomRequest(const UScriptStruct* RequestType, const void* RequestData, const FGuid& DestinationEndpointId, const TSharedRef<IConcertSessionCustomResponseHandler>& Handler) override
+		{
+			// Directly dispatch to the server session.
+			ServerMock.DispatchRequest(EndpointId, RequestType, RequestData, Handler);
+		}
+
+		virtual void HandleCustomEvent(const UScriptStruct* EventType, const void* EventData) override
+		{
+			TArray<TSharedPtr<IConcertSessionCustomEventHandler>>* HandlerArrayPtr = CustomEventHandlers.Find(EventType->GetFName());
+			if (HandlerArrayPtr)
+			{
+				FConcertSessionContext DummyContext;
+				for (const TSharedPtr<IConcertSessionCustomEventHandler>& Handler : *HandlerArrayPtr)
+				{
+					Handler->HandleEvent(DummyContext, EventData);
+				}
+			}
+		}
+
+		virtual FGuid GetSessionClientEndpointId() const override
+		{
+			return EndpointId;
+		}
+
+		virtual FGuid GetSessionServerEndpointId() const override
+		{
+			return FGuid(0, 0, 0, 0);
+		}
+		
+		
+		virtual FOnConcertClientSessionTick& OnTick() override { return OnTickDelegate; }
+
+	private:
+		FOnConcertClientSessionTick OnTickDelegate;
+		FConcertServerSessionMock& ServerMock;
+		FGuid EndpointId;
+		TMap<FName, TArray<TSharedPtr<IConcertSessionCustomEventHandler>>> CustomEventHandlers;
+	};
+
+	/** Base class to perform Concert data store client/server tests. */
+	class FConcertClientServerCommunicationTest : public FAutomationTestBase
+	{
+	public:
+		struct FClientInfo
+		{
+			FClientInfo(const FGuid& ClientEndPointId, FConcertServerSessionMock& Server)
+				: ClientSessionMock(MakeShared<FConcertClientSessionMock>(ClientEndPointId, Server))
+			{}
+
+			const TSharedRef<FConcertClientSessionBaseMock> ClientSessionMock;
+		};
+
+		FConcertClientServerCommunicationTest(const FString& InName, const bool bInComplexTask)
+			: FAutomationTestBase(InName, bInComplexTask)
+		{}
+
+		FClientInfo& ConnectClient()
+		{
+			const FGuid ClientEndpointId(0, 0, 0, Clients.Num() + 1); // {0, 0, 0, 0} is used by the server.
+			Clients.Add(MakeUnique<FClientInfo>(ClientEndpointId, *ServerSessionMock));
+			ServerSessionMock->ConnectClient(ClientEndpointId, *(Clients.Last()->ClientSessionMock));
+			return *Clients.Last();
+		}
+
+		void InitServer()
+		{
+			// Reset everything to be able to rerun the tests. The test framework doesn't destruct/reconstruct this object
+			// at every run, so just ensure we start with a clean state.
+			Clients.Empty();
+			ServerSessionMock = MakeShared<FConcertServerSessionMock>();
+		}
+
+		const TSharedPtr<FConcertServerSessionMock>& GetServerSessionMock() const { return ServerSessionMock; }
+
+	private:
+		TSharedPtr<FConcertServerSessionMock> ServerSessionMock;
+		TArray<TUniquePtr<FClientInfo>> Clients;
+	};
+}
\ No newline at end of file
