Fix failing unit test for Concert Replication

diff --git a/Engine/Plugins/Developer/Concert/ConcertSync/ConcertSyncTest/Source/ConcertSyncTest/Private/Replication/ReplicationHandshakeTests.cpp b/Engine/Plugins/Developer/Concert/ConcertSync/ConcertSyncTest/Source/ConcertSyncTest/Private/Replication/ReplicationHandshakeTests.cpp
--- a/Engine/Plugins/Developer/Concert/ConcertSync/ConcertSyncTest/Source/ConcertSyncTest/Private/Replication/ReplicationHandshakeTests.cpp
+++ b/Engine/Plugins/Developer/Concert/ConcertSync/ConcertSyncTest/Source/ConcertSyncTest/Private/Replication/ReplicationHandshakeTests.cpp
@@ -122,15 +122,16 @@ namespace UE::ConcertSyncTests
 			});
 		// 2.6 Second client is allowed to replicate different properties on the same object
 		FReplicationStreamDescription NonOverlappingProperties;
+		NonOverlappingProperties.BaseDescription.Identifier = FGuid::NewGuid();
 		FReplicatedObjectInfo StaticMeshComponentInfo_Secondary;
 		StaticMeshComponentInfo_Secondary.ClassPath = UStaticMeshComponent::StaticClass();
 		StaticMeshComponentInfo_Secondary.PropertySelection.ReplicatedProperties.Add(MinLODProperty);
 		NonOverlappingProperties.BaseDescription.ReplicationMap.ReplicatedObjects.Add(PathToSomeActorComponent, StaticMeshComponentInfo_Secondary);
 		ClientReplicationManager_Secondary->JoinReplicationSession({ ClientDescription, { NonOverlappingProperties } })
-		.Next([&](const FJoinReplicatedSessionResult& Result)
-		{
-			TestTrue(TEXT("Two clients can replicate differing properties on the same object"), Result.ErrorCode == EJoinReplicationErrorCode::Success);
-		});
+			.Next([&](const FJoinReplicatedSessionResult& Result)
+			{
+				TestTrue(TEXT("Two clients can replicate differing properties on the same object"), Result.ErrorCode == EJoinReplicationErrorCode::Success);
+			});
 		
 		return true;
 	}
diff --git a/Engine/Plugins/Developer/Concert/ConcertSync/ConcertSyncTest/Source/ConcertSyncTest/Private/Util/ConcertMocks.h b/Engine/Plugins/Developer/Concert/ConcertSync/ConcertSyncTest/Source/ConcertSyncTest/Private/Util/ConcertMocks.h
--- a/Engine/Plugins/Developer/Concert/ConcertSync/ConcertSyncTest/Source/ConcertSyncTest/Private/Util/ConcertMocks.h
+++ b/Engine/Plugins/Developer/Concert/ConcertSync/ConcertSyncTest/Source/ConcertSyncTest/Private/Util/ConcertMocks.h
@@ -276,6 +276,12 @@ namespace UE::ConcertSyncTests
 			}
 		}
 
+		virtual bool FindSessionClient(const FGuid& EndpointId, FConcertSessionClientInfo& OutSessionClientInfo) const override
+		{
+			// Need to mock this in case test fails the internal server logic may log the client name which made an invalid request.
+			return true;
+		}
+		
 		virtual FOnConcertServerSessionClientChanged& OnSessionClientChanged() override
 		{
 			return ConnectionChanged;
