Identify replicated objects by client endpointd ID and stream ID instead of just by stream ID.

diff --git a/Engine/Plugins/Developer/Concert/ConcertSync/ConcertSyncClient/Source/ConcertSyncClient/Private/Replication/Processing/ClientReplicationDataCollector.cpp b/Engine/Plugins/Developer/Concert/ConcertSync/ConcertSyncClient/Source/ConcertSyncClient/Private/Replication/Processing/ClientReplicationDataCollector.cpp
--- a/Engine/Plugins/Developer/Concert/ConcertSync/ConcertSyncClient/Source/ConcertSyncClient/Private/Replication/Processing/ClientReplicationDataCollector.cpp
+++ b/Engine/Plugins/Developer/Concert/ConcertSync/ConcertSyncClient/Source/ConcertSyncClient/Private/Replication/Processing/ClientReplicationDataCollector.cpp
@@ -54,7 +54,7 @@ namespace UE::ConcertSyncClient::Replication
 		}
 	}
 
-	void FClientReplicationDataCollector::ForEachPendingObject(TFunctionRef<void(const ConcertSyncCore::FReplicationStreamObjectID&)> ProcessItemFunc) const
+	void FClientReplicationDataCollector::ForEachPendingObject(TFunctionRef<void(const ConcertSyncCore::FStreamedObjectID&)> ProcessItemFunc) const
 	{
 		for (const TPair<FSoftObjectPath, TArray<FObjectInfo>>& Pair : ObjectsToReplicate)
 		{
@@ -69,7 +69,7 @@ namespace UE::ConcertSyncClient::Replication
 	}
 
 	bool FClientReplicationDataCollector::ExtractReplicationDataForObject(
-		const ConcertSyncCore::FReplicationStreamObjectID& ObjectToProcess,
+		const ConcertSyncCore::FStreamedObjectID& ObjectToProcess,
 		TFunctionRef<void(const FConcertSessionSerializedPayload& Payload)> ProcessCopyable,
 		TFunctionRef<void(FConcertSessionSerializedPayload&& Payload)> ProcessMoveable
 		)
diff --git a/Engine/Plugins/Developer/Concert/ConcertSync/ConcertSyncClient/Source/ConcertSyncClient/Private/Replication/Processing/ClientReplicationDataCollector.h b/Engine/Plugins/Developer/Concert/ConcertSync/ConcertSyncClient/Source/ConcertSyncClient/Private/Replication/Processing/ClientReplicationDataCollector.h
--- a/Engine/Plugins/Developer/Concert/ConcertSync/ConcertSyncClient/Source/ConcertSyncClient/Private/Replication/Processing/ClientReplicationDataCollector.h
+++ b/Engine/Plugins/Developer/Concert/ConcertSync/ConcertSyncClient/Source/ConcertSyncClient/Private/Replication/Processing/ClientReplicationDataCollector.h
@@ -41,9 +41,9 @@ namespace UE::ConcertSyncClient::Replication
 		virtual ~FClientReplicationDataCollector() override;
 
 		//~ Begin IReplicationDataSource Interface
-		virtual void ForEachPendingObject(TFunctionRef<void(const ConcertSyncCore::FReplicationStreamObjectID&)> ProcessItemFunc) const override;
+		virtual void ForEachPendingObject(TFunctionRef<void(const ConcertSyncCore::FStreamedObjectID&)> ProcessItemFunc) const override;
 		virtual int32 NumObjects() const override { return NumTrackedObjects; }
-		virtual bool ExtractReplicationDataForObject(const ConcertSyncCore::FReplicationStreamObjectID& Object, TFunctionRef<void(const FConcertSessionSerializedPayload& Payload)> ProcessCopyable, TFunctionRef<void(FConcertSessionSerializedPayload&& Payload)> ProcessMoveable) override;
+		virtual bool ExtractReplicationDataForObject(const ConcertSyncCore::FStreamedObjectID& Object, TFunctionRef<void(const FConcertSessionSerializedPayload& Payload)> ProcessCopyable, TFunctionRef<void(FConcertSessionSerializedPayload&& Payload)> ProcessMoveable) override;
 		//~ End IReplicationDataSource Interface
 
 	private:
diff --git a/Engine/Plugins/Developer/Concert/ConcertSync/ConcertSyncClient/Source/ConcertSyncClient/Private/Replication/Processing/ClientReplicationDataQueuer.cpp b/Engine/Plugins/Developer/Concert/ConcertSync/ConcertSyncClient/Source/ConcertSyncClient/Private/Replication/Processing/ClientReplicationDataQueuer.cpp
--- a/Engine/Plugins/Developer/Concert/ConcertSync/ConcertSyncClient/Source/ConcertSyncClient/Private/Replication/Processing/ClientReplicationDataQueuer.cpp
+++ b/Engine/Plugins/Developer/Concert/ConcertSync/ConcertSyncClient/Source/ConcertSyncClient/Private/Replication/Processing/ClientReplicationDataQueuer.cpp
@@ -20,10 +20,9 @@ namespace UE::ConcertSyncClient::Replication
 		: ReplicationBridge(ReplicationBridge)
 	{}
 
-	bool FClientReplicationDataQueuer::WantsToAcceptObject(const ConcertSyncCore::FReplicationStreamObjectID& Object) const
+	bool FClientReplicationDataQueuer::WantsToAcceptObject(const ConcertSyncCore::FReplicatedObjectInfo& Object) const
 	{
 		// The server tries not to send objects to are unavailable but it could happen the object becomes unavailable while data is underway
-		// TODO: Check whether client attributes allow this object since they could have changed since
 		const bool bIsAvailable = ReplicationBridge->IsObjectAvailable(Object.Object);
 		return bIsAvailable;
 	}
diff --git a/Engine/Plugins/Developer/Concert/ConcertSync/ConcertSyncClient/Source/ConcertSyncClient/Private/Replication/Processing/ClientReplicationDataQueuer.h b/Engine/Plugins/Developer/Concert/ConcertSync/ConcertSyncClient/Source/ConcertSyncClient/Private/Replication/Processing/ClientReplicationDataQueuer.h
--- a/Engine/Plugins/Developer/Concert/ConcertSync/ConcertSyncClient/Source/ConcertSyncClient/Private/Replication/Processing/ClientReplicationDataQueuer.h
+++ b/Engine/Plugins/Developer/Concert/ConcertSync/ConcertSyncClient/Source/ConcertSyncClient/Private/Replication/Processing/ClientReplicationDataQueuer.h
@@ -20,7 +20,7 @@ namespace UE::ConcertSyncClient::Replication
 		static TSharedRef<FClientReplicationDataQueuer> Make(IConcertClientReplicationBridge* ReplicationBridge, TSharedRef<ConcertSyncCore::FObjectReplicationCache> InReplicationCache);
 		
 		//~ Begin IReplicationCacheUser Interface
-		virtual bool WantsToAcceptObject(const ConcertSyncCore::FReplicationStreamObjectID& Object) const override;
+		virtual bool WantsToAcceptObject(const ConcertSyncCore::FReplicatedObjectInfo& Object) const override;
 		//~ End IReplicationCacheUser Interface
 
 	private:
diff --git a/Engine/Plugins/Developer/Concert/ConcertSync/ConcertSyncCore/Source/ConcertSyncCore/Private/Replication/Processing/ObjectReplicationCache.cpp b/Engine/Plugins/Developer/Concert/ConcertSync/ConcertSyncCore/Source/ConcertSyncCore/Private/Replication/Processing/ObjectReplicationCache.cpp
--- a/Engine/Plugins/Developer/Concert/ConcertSync/ConcertSyncCore/Source/ConcertSyncCore/Private/Replication/Processing/ObjectReplicationCache.cpp
+++ b/Engine/Plugins/Developer/Concert/ConcertSync/ConcertSyncCore/Source/ConcertSyncCore/Private/Replication/Processing/ObjectReplicationCache.cpp
@@ -12,11 +12,11 @@ namespace UE::ConcertSyncCore
 		: ReplicationFormat(MoveTemp(ReplicationFormat))
 	{}
 
-	int32 FObjectReplicationCache::StoreUntilConsumed(const FGuid& OriginStreamId, const FConcertObjectReplicationEvent& ObjectReplicationEvent)
+	int32 FObjectReplicationCache::StoreUntilConsumed(const FGuid& SendingEndpointId, const FGuid& OriginStreamId, const FConcertObjectReplicationEvent& ObjectReplicationEvent)
 	{
 		int32 NumAccepted = 0;
 		
-		const FReplicationStreamObjectID ObjectId{ OriginStreamId, ObjectReplicationEvent.ReplicatedObject };
+		const FReplicatedObjectInfo ObjectId{ { OriginStreamId, ObjectReplicationEvent.ReplicatedObject }, SendingEndpointId };
 		FObjectCache* ObjectCacheBeforeAddition = Cache.Find(ObjectId);
 		if (ObjectCacheBeforeAddition)
 		{
@@ -66,7 +66,7 @@ namespace UE::ConcertSyncCore
 					}
 
 					// ObjectCache may not be found because cache user was unregistered and then destroyed: UnregisterDataCacheUser removes cache users.
-					const FReplicationStreamObjectID ObjectId{ OriginStreamId, LazilyCopiedEventPtr->ReplicatedObject };
+					const FStreamedObjectID ObjectId{ OriginStreamId, LazilyCopiedEventPtr->ReplicatedObject };
 					FObjectCache* ObjectCache = This->Cache.Find(ObjectId);
 					if (!ObjectCache)
 					{
diff --git a/Engine/Plugins/Developer/Concert/ConcertSync/ConcertSyncCore/Source/ConcertSyncCore/Private/Replication/Processing/ObjectReplicationProcessor.cpp b/Engine/Plugins/Developer/Concert/ConcertSync/ConcertSyncCore/Source/ConcertSyncCore/Private/Replication/Processing/ObjectReplicationProcessor.cpp
--- a/Engine/Plugins/Developer/Concert/ConcertSync/ConcertSyncCore/Source/ConcertSyncCore/Private/Replication/Processing/ObjectReplicationProcessor.cpp
+++ b/Engine/Plugins/Developer/Concert/ConcertSync/ConcertSyncCore/Source/ConcertSyncCore/Private/Replication/Processing/ObjectReplicationProcessor.cpp
@@ -14,7 +14,7 @@ namespace UE::ConcertSyncCore
 	void FObjectReplicationProcessor::ProcessObjects(float TimeBudget)
 	{
 		// TODO: Respect time budget and prioritize objects
-		DataSource->ForEachPendingObject([this](const FReplicationStreamObjectID& ObjectInfo)
+		DataSource->ForEachPendingObject([this](const FStreamedObjectID& ObjectInfo)
 		{
 			ProcessObject({ ObjectInfo });
 		});
diff --git a/Engine/Plugins/Developer/Concert/ConcertSync/ConcertSyncCore/Source/ConcertSyncCore/Private/Replication/Processing/ObjectReplicationReceiver.cpp b/Engine/Plugins/Developer/Concert/ConcertSync/ConcertSyncCore/Source/ConcertSyncCore/Private/Replication/Processing/ObjectReplicationReceiver.cpp
--- a/Engine/Plugins/Developer/Concert/ConcertSync/ConcertSyncCore/Source/ConcertSyncCore/Private/Replication/Processing/ObjectReplicationReceiver.cpp
+++ b/Engine/Plugins/Developer/Concert/ConcertSync/ConcertSyncCore/Source/ConcertSyncCore/Private/Replication/Processing/ObjectReplicationReceiver.cpp
@@ -32,7 +32,7 @@ namespace UE::ConcertSyncCore
 			NumObjects += StreamEvent.ReplicatedObjects.Num();
 			for (const FConcertObjectReplicationEvent& ObjectEvent : StreamEvent.ReplicatedObjects)
 			{
-				const int32 NumAccepted = ReplicationCache->StoreUntilConsumed(StreamEvent.StreamId, ObjectEvent);
+				const int32 NumAccepted = ReplicationCache->StoreUntilConsumed(SessionContext.SourceEndpointId, StreamEvent.StreamId, ObjectEvent);
 				NumCacheUsages += NumAccepted;
 				NumOfAcceptedObjects += NumAccepted == 0 ? 0 : 1;
 			}
diff --git a/Engine/Plugins/Developer/Concert/ConcertSync/ConcertSyncCore/Source/ConcertSyncCore/Private/Replication/Processing/ReplicationDataQueuer.cpp b/Engine/Plugins/Developer/Concert/ConcertSync/ConcertSyncCore/Source/ConcertSyncCore/Private/Replication/Processing/ReplicationDataQueuer.cpp
--- a/Engine/Plugins/Developer/Concert/ConcertSync/ConcertSyncCore/Source/ConcertSyncCore/Private/Replication/Processing/ReplicationDataQueuer.cpp
+++ b/Engine/Plugins/Developer/Concert/ConcertSync/ConcertSyncCore/Source/ConcertSyncCore/Private/Replication/Processing/ReplicationDataQueuer.cpp
@@ -14,7 +14,7 @@ namespace UE::ConcertSyncCore
 		}
 	}
 
-	void FReplicationDataQueuer::ForEachPendingObject(TFunctionRef<void(const FReplicationStreamObjectID&)> ProcessItemFunc) const
+	void FReplicationDataQueuer::ForEachPendingObject(TFunctionRef<void(const FStreamedObjectID&)> ProcessItemFunc) const
 	{
 		for (auto It = PendingEvents.CreateConstIterator(); It; ++It)
 		{
@@ -28,7 +28,7 @@ namespace UE::ConcertSyncCore
 	}
 
 	bool FReplicationDataQueuer::ExtractReplicationDataForObject(
-		const FReplicationStreamObjectID& Object,
+		const FStreamedObjectID& Object,
 		TFunctionRef<void(const FConcertSessionSerializedPayload& Payload)> ProcessCopyable,
 		TFunctionRef<void(FConcertSessionSerializedPayload&& Payload)> ProcessMoveable
 		)
@@ -45,7 +45,7 @@ namespace UE::ConcertSyncCore
 		return true;
 	}
 
-	void FReplicationDataQueuer::OnDataCached(const FReplicationStreamObjectID& Object, TSharedRef<const FConcertObjectReplicationEvent> Data)
+	void FReplicationDataQueuer::OnDataCached(const FReplicatedObjectInfo& Object, TSharedRef<const FConcertObjectReplicationEvent> Data)
 	{
 		PendingEvents.Add(Object, MoveTemp(Data));
 	}
diff --git a/Engine/Plugins/Developer/Concert/ConcertSync/ConcertSyncCore/Source/ConcertSyncCore/Private/Replication/ReplicationStreamObjectID.cpp b/Engine/Plugins/Developer/Concert/ConcertSync/ConcertSyncCore/Source/ConcertSyncCore/Private/Replication/ReplicationStreamObjectID.cpp
--- a/Engine/Plugins/Developer/Concert/ConcertSync/ConcertSyncCore/Source/ConcertSyncCore/Private/Replication/ReplicationStreamObjectID.cpp
+++ b/Engine/Plugins/Developer/Concert/ConcertSync/ConcertSyncCore/Source/ConcertSyncCore/Private/Replication/ReplicationStreamObjectID.cpp
@@ -4,8 +4,13 @@
 
 namespace UE::ConcertSyncCore
 {
-	uint32 GetTypeHash(const FReplicationStreamObjectID& StreamObject)
+	uint32 GetTypeHash(const FStreamedObjectID& StreamObject)
 	{
 		return HashCombineFast(GetTypeHash(StreamObject.Object), GetTypeHash(StreamObject.StreamId));
 	}
+	
+	uint32 GetTypeHash(const FReplicatedObjectInfo& StreamObject)
+	{
+		return HashCombineFast(GetTypeHash(StreamObject.SenderEndpointId), GetTypeHash(static_cast<FStreamedObjectID>(StreamObject)));
+	}
 }
\ No newline at end of file
diff --git a/Engine/Plugins/Developer/Concert/ConcertSync/ConcertSyncCore/Source/ConcertSyncCore/Public/Replication/Processing/IReplicationDataSource.h b/Engine/Plugins/Developer/Concert/ConcertSync/ConcertSyncCore/Source/ConcertSyncCore/Public/Replication/Processing/IReplicationDataSource.h
--- a/Engine/Plugins/Developer/Concert/ConcertSync/ConcertSyncCore/Source/ConcertSyncCore/Public/Replication/Processing/IReplicationDataSource.h
+++ b/Engine/Plugins/Developer/Concert/ConcertSync/ConcertSyncCore/Source/ConcertSyncCore/Public/Replication/Processing/IReplicationDataSource.h
@@ -9,7 +9,7 @@
 
 namespace UE::ConcertSyncCore
 {
-	struct FReplicationStreamObjectID;
+	struct FStreamedObjectID;
 	
 	/**
 	 * Responsible for obtaining replication data.
@@ -23,7 +23,7 @@ namespace UE::ConcertSyncCore
 	public:
 		
 		/** Iterates the objects that must be processed for replication (the result of ExtractReplicationDataForObject MAY return something new). */
-		virtual void ForEachPendingObject(TFunctionRef<void(const FReplicationStreamObjectID&)> ProcessItemFunc) const = 0;
+		virtual void ForEachPendingObject(TFunctionRef<void(const FStreamedObjectID&)> ProcessItemFunc) const = 0;
 		/** @return The number of objects ForEachPendingObject would iterate. Can be used, e.g. for Reserve()-ing a container. */
 		virtual int32 NumObjects() const = 0;
 
@@ -40,12 +40,12 @@ namespace UE::ConcertSyncCore
 		 * @return Whether successful. False indicates the call was invalid to make (ForEachPendingObject did not return Object). True indicates success, even if ProcessCopyable was not called.		
 		 */
 		virtual bool ExtractReplicationDataForObject(
-			const FReplicationStreamObjectID& Object,
+			const FStreamedObjectID& Object,
 			TFunctionRef<void(const FConcertSessionSerializedPayload& Payload)> ProcessCopyable,
 			TFunctionRef<void(FConcertSessionSerializedPayload&& Payload)> ProcessMoveable
 			) = 0;
 		/** Util version for callers that only want to read and do not want to store the payload. */
-		bool ExtractReplicationDataForObject(const FReplicationStreamObjectID& Object, TFunctionRef<void(const FConcertSessionSerializedPayload& Payload)> ProcessCopyable)
+		bool ExtractReplicationDataForObject(const FStreamedObjectID& Object, TFunctionRef<void(const FConcertSessionSerializedPayload& Payload)> ProcessCopyable)
 		{
 			return ExtractReplicationDataForObject(Object, ProcessCopyable, [&ProcessCopyable](FConcertSessionSerializedPayload&& Payload){ ProcessCopyable(Payload); });
 		}
diff --git a/Engine/Plugins/Developer/Concert/ConcertSync/ConcertSyncCore/Source/ConcertSyncCore/Public/Replication/Processing/ObjectReplicationCache.h b/Engine/Plugins/Developer/Concert/ConcertSync/ConcertSyncCore/Source/ConcertSyncCore/Public/Replication/Processing/ObjectReplicationCache.h
--- a/Engine/Plugins/Developer/Concert/ConcertSync/ConcertSyncCore/Source/ConcertSyncCore/Public/Replication/Processing/ObjectReplicationCache.h
+++ b/Engine/Plugins/Developer/Concert/ConcertSync/ConcertSyncCore/Source/ConcertSyncCore/Public/Replication/Processing/ObjectReplicationCache.h
@@ -25,7 +25,7 @@ namespace UE::ConcertSyncCore
 	public:
 
 		/** @return Whether this user is interested in data from this object. */
-		virtual bool WantsToAcceptObject(const FReplicationStreamObjectID& Object) const = 0;
+		virtual bool WantsToAcceptObject(const FReplicatedObjectInfo& Object) const = 0;
 
 		/**
 		 * Called when data that is interesting to this user becomes available.
@@ -33,7 +33,7 @@ namespace UE::ConcertSyncCore
 		 * The user can keep hold of Data until it is used, at which point it just let's Data get out of scope.
 		 * If new data is received while this user is referencing Data, Data will be combined to contain any new data.
 		 */
-		virtual void OnDataCached(const FReplicationStreamObjectID& Object, TSharedRef<const FConcertObjectReplicationEvent> Data) = 0;
+		virtual void OnDataCached(const FReplicatedObjectInfo& Object, TSharedRef<const FConcertObjectReplicationEvent> Data) = 0;
 
 		virtual ~IReplicationCacheUser() = default;
 	};
@@ -58,11 +58,12 @@ namespace UE::ConcertSyncCore
 
 		/**
 		 * Called when new data is received for an object and shares it with any IObjectCacheUser that is possibly interested in it.
+		 * @param SendingEndpointId The ID of the client endpoint that sent this data
 		 * @param OriginStreamId The stream from which the object was replicated
 		 * @param ObjectReplicationEvent The data that was replicated
 		 * @return The number of cache users that accepted this event. 
 		 */
-		int32 StoreUntilConsumed(const FGuid& OriginStreamId, const FConcertObjectReplicationEvent& ObjectReplicationEvent);
+		int32 StoreUntilConsumed(const FGuid& SendingEndpointId, const FGuid& OriginStreamId, const FConcertObjectReplicationEvent& ObjectReplicationEvent);
 
 		/** Registers a new user, which will start receiving data for any new data received from now on. */
 		void RegisterDataCacheUser(TSharedRef<IReplicationCacheUser> User);
@@ -91,6 +92,6 @@ namespace UE::ConcertSyncCore
 			TMap<TWeakPtr<IReplicationCacheUser>, TWeakPtr<FConcertObjectReplicationEvent>> DataInUse;
 		};
 		/** Maps every object to the events cached for it. */
-		TMap<FReplicationStreamObjectID, FObjectCache> Cache; 
+		TMap<FStreamedObjectID, FObjectCache> Cache; 
 	};
 }
\ No newline at end of file
diff --git a/Engine/Plugins/Developer/Concert/ConcertSync/ConcertSyncCore/Source/ConcertSyncCore/Public/Replication/Processing/ObjectReplicationProcessor.h b/Engine/Plugins/Developer/Concert/ConcertSync/ConcertSyncCore/Source/ConcertSyncCore/Public/Replication/Processing/ObjectReplicationProcessor.h
--- a/Engine/Plugins/Developer/Concert/ConcertSync/ConcertSyncCore/Source/ConcertSyncCore/Public/Replication/Processing/ObjectReplicationProcessor.h
+++ b/Engine/Plugins/Developer/Concert/ConcertSync/ConcertSyncCore/Source/ConcertSyncCore/Public/Replication/Processing/ObjectReplicationProcessor.h
@@ -37,7 +37,7 @@ namespace UE::ConcertSyncCore
 		struct FObjectProcessArgs
 		{
 			/** Info about the object to process */
-			FReplicationStreamObjectID ObjectInfo;
+			FStreamedObjectID ObjectInfo;
 		};
 
 		FORCEINLINE IReplicationDataSource& GetDataSource() const { return DataSource.Get(); }
diff --git a/Engine/Plugins/Developer/Concert/ConcertSync/ConcertSyncCore/Source/ConcertSyncCore/Public/Replication/Processing/ReplicationDataQueuer.h b/Engine/Plugins/Developer/Concert/ConcertSync/ConcertSyncCore/Source/ConcertSyncCore/Public/Replication/Processing/ReplicationDataQueuer.h
--- a/Engine/Plugins/Developer/Concert/ConcertSync/ConcertSyncCore/Source/ConcertSyncCore/Public/Replication/Processing/ReplicationDataQueuer.h
+++ b/Engine/Plugins/Developer/Concert/ConcertSync/ConcertSyncCore/Source/ConcertSyncCore/Public/Replication/Processing/ReplicationDataQueuer.h
@@ -27,13 +27,13 @@ namespace UE::ConcertSyncCore
 		virtual ~FReplicationDataQueuer() override;
 
 		//~ Begin IReplicationDataSource Interface
-		virtual void ForEachPendingObject(TFunctionRef<void(const FReplicationStreamObjectID&)> ProcessItemFunc) const override;
+		virtual void ForEachPendingObject(TFunctionRef<void(const FStreamedObjectID&)> ProcessItemFunc) const override;
 		virtual int32 NumObjects() const override;
-		virtual bool ExtractReplicationDataForObject(const FReplicationStreamObjectID& Object, TFunctionRef<void(const FConcertSessionSerializedPayload& Payload)> ProcessCopyable, TFunctionRef<void(FConcertSessionSerializedPayload&& Payload)> ProcessMoveable) override;
+		virtual bool ExtractReplicationDataForObject(const FStreamedObjectID& Object, TFunctionRef<void(const FConcertSessionSerializedPayload& Payload)> ProcessCopyable, TFunctionRef<void(FConcertSessionSerializedPayload&& Payload)> ProcessMoveable) override;
 		//~ End IReplicationDataSource Interface
 		
 		//~ Begin IReplicationCacheUser Interface
-		virtual void OnDataCached(const FReplicationStreamObjectID& Object, TSharedRef<const FConcertObjectReplicationEvent> Data) override;
+		virtual void OnDataCached(const FReplicatedObjectInfo& Object, TSharedRef<const FConcertObjectReplicationEvent> Data) override;
 		//~ End IReplicationCacheUser Interface
 
 	protected:
@@ -44,7 +44,7 @@ namespace UE::ConcertSyncCore
 	private:
 
 		/** Stores events as they are received. */
-		TMap<FReplicationStreamObjectID, TSharedPtr<const FConcertObjectReplicationEvent>> PendingEvents;
+		TMap<FStreamedObjectID, TSharedPtr<const FConcertObjectReplicationEvent>> PendingEvents;
 
 		/** Provides us with replication events and shares them effectively. */
 		TSharedPtr<FObjectReplicationCache> ReplicationCache;
diff --git a/Engine/Plugins/Developer/Concert/ConcertSync/ConcertSyncCore/Source/ConcertSyncCore/Public/Replication/ReplicationStreamObjectID.h b/Engine/Plugins/Developer/Concert/ConcertSync/ConcertSyncCore/Source/ConcertSyncCore/Public/Replication/ReplicationStreamObjectID.h
--- a/Engine/Plugins/Developer/Concert/ConcertSync/ConcertSyncCore/Source/ConcertSyncCore/Public/Replication/ReplicationStreamObjectID.h
+++ b/Engine/Plugins/Developer/Concert/ConcertSync/ConcertSyncCore/Source/ConcertSyncCore/Public/Replication/ReplicationStreamObjectID.h
@@ -7,23 +7,41 @@
 
 namespace UE::ConcertSyncCore
 {
-	struct FReplicationStreamObjectID
+	struct FStreamedObjectID
 	{
 		/** The replication stream that produces data for this object. */
 		FGuid StreamId;
 		/** The object to process */
 		FSoftObjectPath Object;
-
-		friend bool operator==(const FReplicationStreamObjectID& Left, const FReplicationStreamObjectID& Right)
+		
+		friend bool operator==(const FStreamedObjectID& Left, const FStreamedObjectID& Right)
 		{
 			return Left.StreamId == Right.StreamId && Left.Object == Right.Object;
 		}
 
-		friend bool operator!=(const FReplicationStreamObjectID& Left, const FReplicationStreamObjectID& Right)
+		friend bool operator!=(const FStreamedObjectID& Left, const FStreamedObjectID& Right)
+		{
+			return !(Left == Right);
+		}
+	};
+	
+	struct FReplicatedObjectInfo : FStreamedObjectID
+	{
+		/** The ID of the endpoint that sent this object. This can be the client or server endpoint depending on who receives it. */
+		FGuid SenderEndpointId;
+		
+		friend bool operator==(const FReplicatedObjectInfo& Left, const FReplicatedObjectInfo& Right)
+		{
+			return Left.SenderEndpointId == Right.SenderEndpointId
+				&& static_cast<const FStreamedObjectID&>(Left) == static_cast<const FStreamedObjectID&>(Right);
+		}
+
+		friend bool operator!=(const FReplicatedObjectInfo& Left, const FReplicatedObjectInfo& Right)
 		{
 			return !(Left == Right);
 		}
 	};
 
-	CONCERTSYNCCORE_API uint32 GetTypeHash(const FReplicationStreamObjectID& StreamObject);
+	CONCERTSYNCCORE_API uint32 GetTypeHash(const FStreamedObjectID& StreamObject);
+	CONCERTSYNCCORE_API uint32 GetTypeHash(const FReplicatedObjectInfo& StreamObject);
 }
\ No newline at end of file
diff --git a/Engine/Plugins/Developer/Concert/ConcertSync/ConcertSyncServer/Source/ConcertSyncServer/Private/Replication/ConcertReplicationClient.cpp b/Engine/Plugins/Developer/Concert/ConcertSync/ConcertSyncServer/Source/ConcertSyncServer/Private/Replication/ConcertReplicationClient.cpp
--- a/Engine/Plugins/Developer/Concert/ConcertSync/ConcertSyncServer/Source/ConcertSyncServer/Private/Replication/ConcertReplicationClient.cpp
+++ b/Engine/Plugins/Developer/Concert/ConcertSync/ConcertSyncServer/Source/ConcertSyncServer/Private/Replication/ConcertReplicationClient.cpp
@@ -15,7 +15,7 @@ namespace UE::ConcertSyncServer::Replication
 	)
 		: StreamDescriptions(MoveTemp(StreamDescriptions))
 		, ClientEndpointId(ClientEndpointId)
-		, EventQueue(FServerReplicationDataQueuer::Make(*this, MoveTemp(ReplicationCache)))
+		, EventQueue(FServerReplicationDataQueuer::Make(ClientEndpointId, MoveTemp(ReplicationCache)))
 		, DataRelay(ClientEndpointId, MoveTemp(Session), EventQueue)
 	{}
 
diff --git a/Engine/Plugins/Developer/Concert/ConcertSync/ConcertSyncServer/Source/ConcertSyncServer/Private/Replication/Processing/ServerReplicationDataQueuer.cpp b/Engine/Plugins/Developer/Concert/ConcertSync/ConcertSyncServer/Source/ConcertSyncServer/Private/Replication/Processing/ServerReplicationDataQueuer.cpp
--- a/Engine/Plugins/Developer/Concert/ConcertSync/ConcertSyncServer/Source/ConcertSyncServer/Private/Replication/Processing/ServerReplicationDataQueuer.cpp
+++ b/Engine/Plugins/Developer/Concert/ConcertSync/ConcertSyncServer/Source/ConcertSyncServer/Private/Replication/Processing/ServerReplicationDataQueuer.cpp
@@ -8,33 +8,24 @@
 namespace UE::ConcertSyncServer::Replication
 {
 	TSharedRef<FServerReplicationDataQueuer> FServerReplicationDataQueuer::Make(
-		const FConcertReplicationClient& OwningClient,
+		const FGuid& OwningClientEndpointId,
 		TSharedRef<ConcertSyncCore::FObjectReplicationCache> InReplicationCache
 		)
 	{
-		TSharedRef<FServerReplicationDataQueuer> Result = MakeShared<FServerReplicationDataQueuer>(OwningClient);
+		TSharedRef<FServerReplicationDataQueuer> Result = MakeShared<FServerReplicationDataQueuer>(OwningClientEndpointId);
 		Result->BindToCache(MoveTemp(InReplicationCache));
 		return Result;
 	}
 
-	bool FServerReplicationDataQueuer::WantsToAcceptObject(const ConcertSyncCore::FReplicationStreamObjectID& Object) const
+	bool FServerReplicationDataQueuer::WantsToAcceptObject(const ConcertSyncCore::FReplicatedObjectInfo& Object) const
 	{
 		// Do not send back the data to the client that generated it
-		const bool bWasSentByThisClient = OwningClientStreamIds.Contains(Object.StreamId);
-		
-		// TODO: For now accept all objects from other clients. In the future, clients can specify what data they want to receive using client attributes.
+		const bool bWasSentByThisClient = OwningClientEndpointId == Object.SenderEndpointId;
+		// TODO DP: For now accept all objects from other clients. In the future, clients can specify what data they want to receive using client attributes.
 		return !bWasSentByThisClient;
 	}
 
-	FServerReplicationDataQueuer::FServerReplicationDataQueuer(const FConcertReplicationClient& OwningClient)
-		: OwningClientStreamIds([&OwningClient]()
-		{
-			TSet<FGuid> Streams;
-			for (const FReplicationStreamDescription& Description : OwningClient.GetStreamDescriptions())
-			{
-				Streams.Add(Description.BaseDescription.Identifier);
-			}
-			return Streams;
-		}())
+	FServerReplicationDataQueuer::FServerReplicationDataQueuer(const FGuid& OwningClientEndpointId)
+		: OwningClientEndpointId(OwningClientEndpointId)
 	{}
 }
diff --git a/Engine/Plugins/Developer/Concert/ConcertSync/ConcertSyncServer/Source/ConcertSyncServer/Private/Replication/Processing/ServerReplicationDataQueuer.h b/Engine/Plugins/Developer/Concert/ConcertSync/ConcertSyncServer/Source/ConcertSyncServer/Private/Replication/Processing/ServerReplicationDataQueuer.h
--- a/Engine/Plugins/Developer/Concert/ConcertSync/ConcertSyncServer/Source/ConcertSyncServer/Private/Replication/Processing/ServerReplicationDataQueuer.h
+++ b/Engine/Plugins/Developer/Concert/ConcertSync/ConcertSyncServer/Source/ConcertSyncServer/Private/Replication/Processing/ServerReplicationDataQueuer.h
@@ -21,18 +21,18 @@ namespace UE::ConcertSyncServer::Replication
 		friend class SharedPointerInternals::TIntrusiveReferenceController;
 	public:
 		
-		static TSharedRef<FServerReplicationDataQueuer> Make(const FConcertReplicationClient& OwningClient, TSharedRef<ConcertSyncCore::FObjectReplicationCache> InReplicationCache);
+		static TSharedRef<FServerReplicationDataQueuer> Make(const FGuid& OwningClientEndpointId, TSharedRef<ConcertSyncCore::FObjectReplicationCache> InReplicationCache);
 		
 		//~ Begin IReplicationCacheUser Interface
-		virtual bool WantsToAcceptObject(const ConcertSyncCore::FReplicationStreamObjectID& Object) const override;
+		virtual bool WantsToAcceptObject(const ConcertSyncCore::FReplicatedObjectInfo& Object) const override;
 		//~ End IReplicationCacheUser Interface
 
 	private:
 
-		FServerReplicationDataQueuer(const FConcertReplicationClient& OwningClient);
+		FServerReplicationDataQueuer(const FGuid& OwningClientEndpointId);
 
 		/** The client for which this FServerReplicationDataQueuer exists. */
-		const TSet<FGuid> OwningClientStreamIds;
+		const FGuid OwningClientEndpointId;
 	};
 }
 
diff --git a/Engine/Plugins/Developer/Concert/ConcertSync/ConcertSyncTest/Source/ConcertSyncTest/Private/Replication/ReplicationCacheTests.cpp b/Engine/Plugins/Developer/Concert/ConcertSync/ConcertSyncTest/Source/ConcertSyncTest/Private/Replication/ReplicationCacheTests.cpp
--- a/Engine/Plugins/Developer/Concert/ConcertSync/ConcertSyncTest/Source/ConcertSyncTest/Private/Replication/ReplicationCacheTests.cpp
+++ b/Engine/Plugins/Developer/Concert/ConcertSync/ConcertSyncTest/Source/ConcertSyncTest/Private/Replication/ReplicationCacheTests.cpp
@@ -60,22 +60,22 @@ namespace UE::ConcertSyncTests
 		public:
 
 			FAutomationTestBase& Test;
-			ConcertSyncCore::FReplicationStreamObjectID AllowedObject;
+			ConcertSyncCore::FStreamedObjectID AllowedObject;
 			EReplicationCacheTestFlags Flags;
 			TSharedPtr<const FConcertObjectReplicationEvent> CachedData;
 
-			FTestReplicationCacheUser(FAutomationTestBase& Test, ConcertSyncCore::FReplicationStreamObjectID AllowedObject, EReplicationCacheTestFlags Flags = EReplicationCacheTestFlags::None)
+			FTestReplicationCacheUser(FAutomationTestBase& Test, ConcertSyncCore::FStreamedObjectID AllowedObject, EReplicationCacheTestFlags Flags = EReplicationCacheTestFlags::None)
 				: Test(Test)
 				, AllowedObject(MoveTemp(AllowedObject))
 				, Flags(Flags)
 			{}
 			
-			virtual bool WantsToAcceptObject(const ConcertSyncCore::FReplicationStreamObjectID& Object) const override
+			virtual bool WantsToAcceptObject(const ConcertSyncCore::FReplicatedObjectInfo& Object) const override
 			{
 				return !EnumHasAnyFlags(Flags, EReplicationCacheTestFlags::NeverReceive) && AllowedObject == Object;
 			}
 			
-			virtual void OnDataCached(const ConcertSyncCore::FReplicationStreamObjectID& Object, TSharedRef<const FConcertObjectReplicationEvent> Data) override
+			virtual void OnDataCached(const ConcertSyncCore::FReplicatedObjectInfo& Object, TSharedRef<const FConcertObjectReplicationEvent> Data) override
 			{
 				if (EnumHasAnyFlags(Flags, EReplicationCacheTestFlags::NeverReceive))
 				{
@@ -123,7 +123,8 @@ namespace UE::ConcertSyncTests
 		// Set up
 		const FGuid StreamId = FGuid::NewGuid();
 		const FSoftObjectPath ObjectPath(TEXT("/Game/World.World:PersistentLevel.StaticMeshActor0"));
-		const ConcertSyncCore::FReplicationStreamObjectID ObjectID{ StreamId, ObjectPath};
+		const FGuid DummySendingClientId = FGuid::NewGuid();
+		const ConcertSyncCore::FReplicatedObjectInfo ObjectID{ { StreamId, ObjectPath }, DummySendingClientId};
 		
 		TSharedRef<ConcertSyncCore::FObjectReplicationCache> Cache = MakeShared<ConcertSyncCore::FObjectReplicationCache>(MakeShared<FTestReplicationFormat>());
 		TSharedRef<FTestReplicationCacheUser> User_NeverConsume = MakeShared<FTestReplicationCacheUser>(*this, ObjectID, EReplicationCacheTestFlags::NeverConsume);
@@ -144,11 +145,11 @@ namespace UE::ConcertSyncTests
 
 		
 		// Tests
-		Cache->StoreUntilConsumed(StreamId, Event_5);
+		Cache->StoreUntilConsumed(DummySendingClientId, StreamId, Event_5);
 		TestTrue(TEXT("Users have same data"), User_NeverConsume->CachedData.Get() == User_ConsumeManually->CachedData.Get());
 		User_ConsumeManually->Consume();
 		
-		Cache->StoreUntilConsumed(StreamId, Event_10);
+		Cache->StoreUntilConsumed(DummySendingClientId, StreamId, Event_10);
 		TestEqual(TEXT("Combined events: 5 and 10"), User_NeverConsume->PeakData(), 15.f);
 		TestEqual(TEXT("Received new data: 10"), User_ConsumeManually->PeakData(), 10.f);
 		
@@ -156,7 +157,7 @@ namespace UE::ConcertSyncTests
 		Cache->UnregisterDataCacheUser(User_ConsumeManually);
 		Cache->RegisterDataCacheUser(User_ConsumeManually);
 		const void* AddressBefore = User_ConsumeManually->GetDataAddress();
-		Cache->StoreUntilConsumed(StreamId, Event_100);
+		Cache->StoreUntilConsumed(DummySendingClientId, StreamId, Event_100);
 		TestEqual(TEXT("Combined events: 5, 10, and 100"), User_NeverConsume->PeakData(), 115.f);
 		TestEqual(TEXT("Received totally new data"), User_ConsumeManually->PeakData(), 100.f);
 		TestNotEqual(TEXT("Received totally new allocated data"), AddressBefore, User_ConsumeManually->GetDataAddress());
