Refactor unit tests for upcoming authority changes

diff --git a/Engine/Plugins/Developer/Concert/ConcertSync/ConcertSyncTest/Source/ConcertSyncTest/Private/Replication/ReplicationCacheTests.cpp b/Engine/Plugins/Developer/Concert/ConcertSync/ConcertSyncTest/Source/ConcertSyncTest/Private/Replication/ReplicationCacheTests.cpp
--- a/Engine/Plugins/Developer/Concert/ConcertSync/ConcertSyncTest/Source/ConcertSyncTest/Private/Replication/ReplicationCacheTests.cpp
+++ b/Engine/Plugins/Developer/Concert/ConcertSync/ConcertSyncTest/Source/ConcertSyncTest/Private/Replication/ReplicationCacheTests.cpp
@@ -5,7 +5,7 @@
 #include "Replication/Processing/ObjectReplicationCache.h"
 #include "Replication/ReplicationStreamObjectID.h"
 #include "TestReflectionObject.h"
-#include "Util/ConcertMocks.h"
+#include "Util/ClientServerCommunicationTest.h"
 
 #include "Misc/AutomationTest.h"
 #include "Misc/Paths.h"
diff --git a/Engine/Plugins/Developer/Concert/ConcertSync/ConcertSyncTest/Source/ConcertSyncTest/Private/Replication/ReplicationHandshakeTests.cpp b/Engine/Plugins/Developer/Concert/ConcertSync/ConcertSyncTest/Source/ConcertSyncTest/Private/Replication/ReplicationHandshakeTests.cpp
--- a/Engine/Plugins/Developer/Concert/ConcertSync/ConcertSyncTest/Source/ConcertSyncTest/Private/Replication/ReplicationHandshakeTests.cpp
+++ b/Engine/Plugins/Developer/Concert/ConcertSync/ConcertSyncTest/Source/ConcertSyncTest/Private/Replication/ReplicationHandshakeTests.cpp
@@ -1,6 +1,6 @@
 // Copyright Epic Games, Inc. All Rights Reserved.
 
-#include "Util/ConcertMocks.h"
+#include "Util/ClientServerCommunicationTest.h"
 
 #include "Replication/Data/ReplicationClientDescription.h"
 #include "Replication/Data/ReplicationStreamDescription.h"
@@ -16,7 +16,7 @@
 #include "Misc/Paths.h"
 #include "HAL/FileManager.h"
 
-namespace UE::ConcertSyncTests
+namespace UE::ConcertSyncTests::Replication
 {
 	/**
 	 * Tests the handshake for joining and leaving a replication session.
diff --git a/Engine/Plugins/Developer/Concert/ConcertSync/ConcertSyncTest/Source/ConcertSyncTest/Private/Replication/ReplicationSendReceiveTests.cpp b/Engine/Plugins/Developer/Concert/ConcertSync/ConcertSyncTest/Source/ConcertSyncTest/Private/Replication/ReplicationSendReceiveTests.cpp
--- a/Engine/Plugins/Developer/Concert/ConcertSync/ConcertSyncTest/Source/ConcertSyncTest/Private/Replication/ReplicationSendReceiveTests.cpp
+++ b/Engine/Plugins/Developer/Concert/ConcertSync/ConcertSyncTest/Source/ConcertSyncTest/Private/Replication/ReplicationSendReceiveTests.cpp
@@ -1,5 +1,7 @@
 // Copyright Epic Games, Inc. All Rights Reserved.
 
+#include "Util/SendReceiveObjectTestBase.h"
+
 #include "Replication/Data/ReplicationStreamDescription.h"
 #include "Replication/IConcertClientReplicationManager.h"
 #include "Replication/IConcertServerReplicationManager.h"
@@ -9,7 +11,7 @@
 #include "Replication/ReplicationTestInterface.h"
 #include "TestReflectionObject.h"
 #include "Util/ConcertClientReplicationBridgeMock.h"
-#include "Util/ConcertMocks.h"
+#include "Util/ClientServerCommunicationTest.h"
 
 #include "Misc/AutomationTest.h"
 #include "Misc/Paths.h"
@@ -17,185 +19,43 @@
 #include "UObject/UObjectGlobals.h"
 #include "UObject/Package.h"
 
-namespace UE::ConcertSyncTests
+namespace UE::ConcertSyncTests::Replication
 {
-	static void SetTestValues(UTestReflectionObject& Object);
-	static void SetDifferentValues(UTestReflectionObject& Object);
-	static void TestEqualTestValues(UTestReflectionObject& Object, FAutomationTestBase& Test);
-
-	static void TickClient(FConcertClientServerCommunicationTest::FClientInfo& Client);
-	static void TickServer(FConcertServerSessionMock& ServerSession);
-
 	/**
 	 * Tests replicating data from sender client > server > receiver client.
 	 */
-	IMPLEMENT_CUSTOM_SIMPLE_AUTOMATION_TEST(FSendReceiveLogicTests, FConcertClientServerCommunicationTest, "Concert.Replication.SendReceiveLogic", EAutomationTestFlags::EditorContext | EAutomationTestFlags::EngineFilter);
-	bool FSendReceiveLogicTests::RunTest(const FString& Parameters)
+	IMPLEMENT_CUSTOM_SIMPLE_AUTOMATION_TEST(FSendReceiveFlowTests, FSendReceiveObjectTestBase, "Concert.Replication.SendReceiveFlow", EAutomationTestFlags::EditorContext | EAutomationTestFlags::EngineFilter);
+	bool FSendReceiveFlowTests::RunTest(const FString& Parameters)
 	{
-		using namespace ConcertSyncClient::TestInterface;
-		using namespace ConcertSyncServer::TestInterface;
-		using namespace ConcertSyncClient::Replication;
-		using namespace ConcertSyncServer::Replication;
-
 		// 1. Init
+		SetUpClientAndServer();
 		
-		// Server
-		InitServer();
-		const TSharedPtr<FConcertServerSessionMock>& ServerSession = GetServerSessionMock();
-		const TSharedRef<IConcertServerReplicationManager> ServerReplicationManager = CreateServerReplicationManager(ServerSession.ToSharedRef());
-		
-		// Client Receiver
-		FClientInfo& Client_Receiver = ConnectClient();
-		const TSharedRef<FConcertClientReplicationBridgeMock> BridgeMock_Receiver = MakeShared<FConcertClientReplicationBridgeMock>();
-		const TSharedPtr<IConcertClientSession>& ClientSession_Receiver = Client_Receiver.ClientSessionMock;
-		const TSharedRef<IConcertClientReplicationManager> ClientReplicationManager_Receiver = CreateClientReplicationManager(ClientSession_Receiver.ToSharedRef(), &BridgeMock_Receiver.Get());
-		
-		// Client Sender
-		FClientInfo& Client_Sender = ConnectClient();
-		const TSharedRef<FConcertClientReplicationBridgeMock> BridgeMock_Sender = MakeShared<FConcertClientReplicationBridgeMock>();
-		const TSharedPtr<IConcertClientSession>& ClientSession_Sender = Client_Sender.ClientSessionMock;
-		const TSharedRef<IConcertClientReplicationManager> ClientReplicationManager_Sender = CreateClientReplicationManager(ClientSession_Sender.ToSharedRef(), &BridgeMock_Sender.Get());
-
-		// Fake replicated object
-		UTestReflectionObject* TestObject = NewObject<UTestReflectionObject>(GetTransientPackage());
-		const FSoftObjectPath TestObjectPath = TestObject;
-		
-		// 1.1 Sender offers to send all UTestReflectionObject properties
-		const FGuid SenderStreamId = FGuid::NewGuid();
-		{
-			FReplicatedObjectInfo AllProperties { UTestReflectionObject::StaticClass() };
-			ConcertSyncCore::PropertyChain::ForEachReplicatableConcertProperty(*UTestReflectionObject::StaticClass(), [&AllProperties](FConcertPropertyChain&& Chain)
-			{
-				AllProperties.PropertySelection.ReplicatedProperties.Emplace(MoveTemp(Chain));
-				return EBreakBehavior::Continue;
-			});
-			
-			FJoinReplicatedSessionArgs SenderJoinArgs;
-			FReplicationStreamDescription SendingStream;
-			SendingStream.BaseDescription.Identifier = SenderStreamId;
-			SendingStream.BaseDescription.ReplicationMap.ReplicatedObjects.Add(TestObjectPath, AllProperties);
-			SenderJoinArgs.Streams.Add(SendingStream);
-			ClientReplicationManager_Sender->JoinReplicationSession(SenderJoinArgs)
-				.Next([this](const FJoinReplicatedSessionResult& Result){ TestTrue(TEXT("Sender joined"), Result.ErrorCode == EJoinReplicationErrorCode::Success); });
-		}
-
-		// 1.2 Receiver accepts all properties
-		{
-			FJoinReplicatedSessionArgs ReceiverJoinArgs;
-			// TODO: When we add client attributes, this test must be updated with attributes that allow receiving all data
-			ClientReplicationManager_Receiver->JoinReplicationSession(ReceiverJoinArgs)
-				.Next([this](const FJoinReplicatedSessionResult& Result){ TestTrue(TEXT("Receiver joined"), Result.ErrorCode == EJoinReplicationErrorCode::Success); });
-		}
-		
-		// 1.3 Prepare test data
+		// 2. Send data
 		bool bHasServerReceivedData = false;
 		bool bHasClientReceivedData = false;
-		auto TestReplicationData_Server = [&](const FConcertSessionContext& Context, const FConcertBatchReplicationEvent& Event)
+		auto OnServerReceive = [this, &bHasServerReceivedData](const FConcertSessionContext& Context, const FConcertBatchReplicationEvent& Event) mutable
 		{
 			if (bHasServerReceivedData)
 			{
 				AddError(TEXT("Server was expected to receive data exactly once!"));
-				return;
 			}
 			bHasServerReceivedData = true;
-			
-			TestEqual(TEXT("Server received 1 stream"), Event.Streams.Num(), 1);
-			TestEqual(TEXT("Server received 1 object"), Event.Streams.IsEmpty() ? 0 : Event.Streams[0].ReplicatedObjects.Num() , 1);
-			TestEqual(TEXT("Server received from correct stream"), Event.Streams.IsEmpty() ? FGuid{} : Event.Streams[0].StreamId, SenderStreamId);
-			const FSoftObjectPath ObjectPath = Event.Streams.IsEmpty() || Event.Streams[0].ReplicatedObjects.IsEmpty() ? FSoftObjectPath{} : Event.Streams[0].ReplicatedObjects[0].ReplicatedObject;
-			TestEqual(TEXT("Server's received object has correct path"), ObjectPath, TestObjectPath);
 		};
-		auto TestReplicationData_Client_Sender = [&](const FConcertSessionContext& Context, const FConcertBatchReplicationEvent& Event)
-        {
-        	// There was a bug where client would receive its own data...
-			AddError(TEXT("Client 1 was not supposed to receive data!"));
-        };
-		auto TestReplicationData_Client_Receiver = [&](const FConcertSessionContext& Context, const FConcertBatchReplicationEvent& Event)
+		auto OnClientReceive = [this, &bHasClientReceivedData](const FConcertSessionContext& Context, const FConcertBatchReplicationEvent& Event) mutable
 		{
 			if (bHasClientReceivedData)
 			{
 				AddError(TEXT("Client 2 was expected to receive data exactly once!"));
-				return;
 			}
 			bHasClientReceivedData = true;
-			
-			TestEqual(TEXT("Client 2 received 1 stream"), Event.Streams.Num(), 1);
-			TestEqual(TEXT("Client 2 received 1 object"), Event.Streams.IsEmpty() ? 0 : Event.Streams[0].ReplicatedObjects.Num() , 1);
-			TestEqual(TEXT("Client 2 received from correct stream"), Event.Streams.IsEmpty() ? FGuid{} : Event.Streams[0].StreamId, SenderStreamId);
-			const FSoftObjectPath ObjectPath = Event.Streams.IsEmpty() || Event.Streams[0].ReplicatedObjects.IsEmpty() ? FSoftObjectPath{} : Event.Streams[0].ReplicatedObjects[0].ReplicatedObject;
-			TestEqual(TEXT("Client 2's received object has correct path"), ObjectPath, TestObjectPath);
 		};
-		ServerSession->RegisterCustomEventHandler<FConcertBatchReplicationEvent>(TestReplicationData_Server);
-		Client_Sender.ClientSessionMock->RegisterCustomEventHandler<FConcertBatchReplicationEvent>(TestReplicationData_Client_Sender);
-		Client_Receiver.ClientSessionMock->RegisterCustomEventHandler<FConcertBatchReplicationEvent>(TestReplicationData_Client_Receiver);
-
-		// 1.4 Replication bridge
-		// Bridge is responsible for telling client-side replication system about existing objects
-		BridgeMock_Sender->InjectAvailableObject(*TestObject); 
-		BridgeMock_Receiver->InjectAvailableObject(*TestObject);
-
-		
+		SimulateSenderToReceiver(OnServerReceive, OnClientReceive);
 
-		
-		// 2. Send data
-		// Note that the systems operate on the same UObject instance. This is required since the systems internally use object path.
-		// For this reason we change the object after sending.
-		
-		// 2.1 Send from sender to server
-		SetTestValues(*TestObject);
-		TickClient(Client_Sender);
-		
-		// 2.2 Forward from server to receiver
-		TickServer(*ServerSession);
-		
-		// 2.3 Receive from server
-		SetDifferentValues(*TestObject);
-		TickClient(Client_Receiver);
-
-
-
-
-		
 		// 3. Test
-		// ... that the data went through server and receiving client ...
 		TestTrue(TEXT("Server received replication event"), bHasServerReceivedData);
 		TestTrue(TEXT("Client 2 received replication event"), bHasClientReceivedData);
-		// ... and the data was applied correctly to the object
 		TestEqualTestValues(*TestObject, *this);
+		
 		return true;
 	}
-
-	// TODO: Add more properties here
-	constexpr float SentFloat = 42.f;
-	const FVector SentVector = { 21.f, 84.f, -1.f };
-
-	static void SetTestValues(UTestReflectionObject& Object)
-	{
-		Object.Float = SentFloat;
-		Object.Vector = SentVector;
-	}
-	
-	static void SetDifferentValues(UTestReflectionObject& Object)
-	{
-		Object.Float *= -1.f;
-		Object.Vector *= -1.f;
-	}
-	
-	static void TestEqualTestValues(UTestReflectionObject& Object, FAutomationTestBase& Test)
-	{
-		Test.TestEqual(TEXT("Float"), Object.Float, SentFloat);
-		Test.TestEqual(TEXT("Vector"), Object.Vector, SentVector);
-	}
-	
-	constexpr float FakeDeltaTime = 1.f / 60.f;
-	
-	static void TickClient(FConcertClientServerCommunicationTest::FClientInfo& Client)
-	{
-		Client.ClientSessionMock->OnTick().Broadcast(*Client.ClientSessionMock, FakeDeltaTime);
-	}
-
-	static void TickServer(FConcertServerSessionMock& ServerSession)
-	{
-		ServerSession.OnTick().Broadcast(ServerSession, FakeDeltaTime);
-	}
 }
diff --git a/Engine/Plugins/Developer/Concert/ConcertSync/ConcertSyncTest/Source/ConcertSyncTest/Private/Replication/Util/ConcertClientReplicationBridgeMock.h b/Engine/Plugins/Developer/Concert/ConcertSync/ConcertSyncTest/Source/ConcertSyncTest/Private/Replication/Util/ConcertClientReplicationBridgeMock.h
new file mode 100644
--- /dev/null
+++ b/Engine/Plugins/Developer/Concert/ConcertSync/ConcertSyncTest/Source/ConcertSyncTest/Private/Replication/Util/ConcertClientReplicationBridgeMock.h
@@ -0,0 +1,58 @@
+// Copyright Epic Games, Inc. All Rights Reserved.
+
+#pragma once
+
+#include "Algo/AnyOf.h"
+#include "Replication/IConcertClientReplicationBridge.h"
+
+namespace UE::ConcertSyncTests::Replication
+{
+	class FConcertClientReplicationBridgeMock : public IConcertClientReplicationBridge
+	{
+	public:
+
+		TArray<FSoftObjectPath> TrackedObjects;
+		TArray<UObject*> AvailableObjects;
+		
+		FConcertClientReplicationBridgeObjectEvent OnObjectDiscoveredDelegate;
+		FConcertClientReplicationBridgeObjectEvent OnObjectRemovedDelegate;
+
+		void InjectAvailableObject(UObject& Object)
+		{
+			AvailableObjects.Add(&Object);
+			OnObjectDiscoveredDelegate.Broadcast(Object);
+		}
+		
+		virtual void PushTrackedObjects(TArrayView<const FSoftObjectPath> InTrackedObjects) override
+		{
+			for (const FSoftObjectPath& Path : InTrackedObjects)
+			{
+				TrackedObjects.Add(Path);
+			}
+		}
+		
+		virtual void PopTrackedObjects(TArrayView<const FSoftObjectPath> InTrackedObjects) override
+		{
+			for (const FSoftObjectPath& Path : InTrackedObjects)
+			{
+				TrackedObjects.RemoveSingle(Path);
+			}
+		}
+
+		virtual bool IsObjectAvailable(const FSoftObjectPath& Path) override { return Algo::AnyOf(AvailableObjects, [&Path](UObject* Object){ return Path == Object; }); }
+		virtual UObject* FindObjectIfAvailable(const FSoftObjectPath& Path) override
+		{
+			for (UObject* Object : AvailableObjects)
+			{
+				if (Path == Object)
+				{
+					return Object;
+				}
+			}
+			return nullptr;
+		}
+
+		virtual FConcertClientReplicationBridgeObjectEvent& OnObjectDiscovered() override { return OnObjectDiscoveredDelegate; }
+		virtual FConcertClientReplicationBridgeObjectEvent& OnObjectHidden() override { return OnObjectRemovedDelegate; }
+	};
+}
diff --git a/Engine/Plugins/Developer/Concert/ConcertSync/ConcertSyncTest/Source/ConcertSyncTest/Private/Replication/Util/SendReceiveObjectTestBase.cpp b/Engine/Plugins/Developer/Concert/ConcertSync/ConcertSyncTest/Source/ConcertSyncTest/Private/Replication/Util/SendReceiveObjectTestBase.cpp
new file mode 100644
--- /dev/null
+++ b/Engine/Plugins/Developer/Concert/ConcertSync/ConcertSyncTest/Source/ConcertSyncTest/Private/Replication/Util/SendReceiveObjectTestBase.cpp
@@ -0,0 +1,123 @@
+// Copyright Epic Games, Inc. All Rights Reserved.
+
+#include "SendReceiveObjectTestBase.h"
+
+#include "ConcertClientReplicationBridgeMock.h"
+#include "Replication/Data/ReplicationStreamDescription.h"
+#include "Replication/IConcertClientReplicationManager.h"
+#include "Replication/IConcertServerReplicationManager.h"
+#include "Replication/Messages/ConcertReplicationEvents.h"
+#include "Replication/Messages/ConcertReplicationHandshakeMessages.h"
+#include "Replication/PropertyChainUtils.h"
+#include "Replication/ReplicationTestInterface.h"
+#include "Replication/TestReflectionObject.h"
+#include "Util/ClientServerCommunicationTest.h"
+
+#include "Misc/AutomationTest.h"
+#include "Misc/Paths.h"
+#include "HAL/FileManager.h"
+#include "UObject/UObjectGlobals.h"
+#include "UObject/Package.h"
+
+namespace UE::ConcertSyncTests::Replication
+{
+	ConcertSyncClient::Replication::FJoinReplicatedSessionArgs FSendReceiveObjectTestBase::CreateSenderArgs()
+	{
+		ConcertSyncClient::Replication::FJoinReplicatedSessionArgs SenderJoinArgs;
+		
+		FReplicatedObjectInfo AllProperties { UTestReflectionObject::StaticClass() };
+		ConcertSyncCore::PropertyChain::ForEachReplicatableConcertProperty(*UTestReflectionObject::StaticClass(), [&AllProperties](FConcertPropertyChain&& Chain)
+		{
+			AllProperties.PropertySelection.ReplicatedProperties.Emplace(MoveTemp(Chain));
+			return EBreakBehavior::Continue;
+		});
+
+		FReplicationStreamDescription SendingStream;
+		SendingStream.BaseDescription.Identifier = SenderStreamId;
+		SendingStream.BaseDescription.ReplicationMap.ReplicatedObjects.Add(TestObject, AllProperties);
+		SenderJoinArgs.Streams.Add(SendingStream);
+		return SenderJoinArgs;
+	}
+
+	ConcertSyncClient::Replication::FJoinReplicatedSessionArgs FSendReceiveObjectTestBase::CreateReceiverArgs()
+	{
+		ConcertSyncClient::Replication::FJoinReplicatedSessionArgs ReceiverJoinArgs;
+		// TODO DP: When we add client attributes, this test must be updated with attributes that allow receiving all data
+		return ReceiverJoinArgs;
+	}
+
+	void FSendReceiveObjectTestBase::SetUpClientAndServer()
+	{
+		// Fake replicated object must be created before call to Super
+		TestObject = NewObject<UTestReflectionObject>(GetTransientPackage());
+
+		FSendReceiveTestBase::SetUpClientAndServer();
+		
+		// Bridge is responsible for telling client-side replication system about existing objects
+		BridgeMock_Sender->InjectAvailableObject(*TestObject); 
+		BridgeMock_Receiver->InjectAvailableObject(*TestObject);
+	}
+	
+	void FSendReceiveObjectTestBase::SimulateSenderToReceiver(
+		TFunctionRef<FReceiveReplicationEventSignature> OnServerReceive,
+		TFunctionRef<FReceiveReplicationEventSignature> OnReceiverClientReceive)
+	{
+		auto TestReplicationData_Server = [this, OnServerReceive](const FConcertSessionContext& Context, const FConcertBatchReplicationEvent& Event)
+		{
+			TestEqual(TEXT("Server received 1 stream"), Event.Streams.Num(), 1);
+			TestEqual(TEXT("Server received 1 object"), Event.Streams.IsEmpty() ? 0 : Event.Streams[0].ReplicatedObjects.Num() , 1);
+			TestEqual(TEXT("Server received from correct stream"), Event.Streams.IsEmpty() ? FGuid{} : Event.Streams[0].StreamId, SenderStreamId);
+			const FSoftObjectPath ObjectPath = Event.Streams.IsEmpty() || Event.Streams[0].ReplicatedObjects.IsEmpty() ? FSoftObjectPath{} : Event.Streams[0].ReplicatedObjects[0].ReplicatedObject;
+			TestEqual(TEXT("Server's received object has correct path"), ObjectPath, FSoftObjectPath(TestObject));
+
+			OnServerReceive(Context, Event);
+		};
+		auto TestReplicationData_Client_Receiver = [this, OnReceiverClientReceive](const FConcertSessionContext& Context, const FConcertBatchReplicationEvent& Event)
+		{
+			TestEqual(TEXT("Client 2 received 1 stream"), Event.Streams.Num(), 1);
+			TestEqual(TEXT("Client 2 received 1 object"), Event.Streams.IsEmpty() ? 0 : Event.Streams[0].ReplicatedObjects.Num() , 1);
+			TestEqual(TEXT("Client 2 received from correct stream"), Event.Streams.IsEmpty() ? FGuid{} : Event.Streams[0].StreamId, SenderStreamId);
+			const FSoftObjectPath ObjectPath = Event.Streams.IsEmpty() || Event.Streams[0].ReplicatedObjects.IsEmpty() ? FSoftObjectPath{} : Event.Streams[0].ReplicatedObjects[0].ReplicatedObject;
+			TestEqual(TEXT("Client 2's received object has correct path"), ObjectPath, FSoftObjectPath(TestObject));
+
+			OnReceiverClientReceive(Context, Event);
+		};
+		const FDelegateHandle ServerHandle = ServerSession->RegisterCustomEventHandler<FConcertBatchReplicationEvent>(TestReplicationData_Server);
+		const FDelegateHandle ClientHandle = Client_Receiver->ClientSessionMock->RegisterCustomEventHandler<FConcertBatchReplicationEvent>(TestReplicationData_Client_Receiver);
+
+		
+		// TestObject is the same UObject on both clients.
+		// Hence we must override test values with SetTestValues and SetDifferentValues.
+		// 1 Sender > Server
+		SetTestValues(*TestObject);
+		TickClient(Client_Sender);
+		
+		// 2 Forward from server to receiver
+		TickServer();
+		
+		// 3 Receive from server
+		SetDifferentValues(*TestObject);
+		TickClient(Client_Receiver);
+
+		ServerSession->UnregisterCustomEventHandler<FConcertBatchReplicationEvent>(ServerHandle);
+		ServerSession->UnregisterCustomEventHandler<FConcertBatchReplicationEvent>(ClientHandle);
+	}
+
+	void FSendReceiveObjectTestBase::SetTestValues(UTestReflectionObject& Object)
+	{
+		Object.Float = SentFloat;
+		Object.Vector = SentVector;
+	}
+	
+	void FSendReceiveObjectTestBase::SetDifferentValues(UTestReflectionObject& Object)
+	{
+		Object.Float *= -1.f;
+		Object.Vector *= -1.f;
+	}
+	
+	void FSendReceiveObjectTestBase::TestEqualTestValues(UTestReflectionObject& Object, FAutomationTestBase& Test)
+	{
+		Test.TestEqual(TEXT("Float"), Object.Float, SentFloat);
+		Test.TestEqual(TEXT("Vector"), Object.Vector, SentVector);
+	}
+}
\ No newline at end of file
diff --git a/Engine/Plugins/Developer/Concert/ConcertSync/ConcertSyncTest/Source/ConcertSyncTest/Private/Replication/Util/SendReceiveObjectTestBase.h b/Engine/Plugins/Developer/Concert/ConcertSync/ConcertSyncTest/Source/ConcertSyncTest/Private/Replication/Util/SendReceiveObjectTestBase.h
new file mode 100644
--- /dev/null
+++ b/Engine/Plugins/Developer/Concert/ConcertSync/ConcertSyncTest/Source/ConcertSyncTest/Private/Replication/Util/SendReceiveObjectTestBase.h
@@ -0,0 +1,64 @@
+// Copyright Epic Games, Inc. All Rights Reserved.
+
+#pragma once
+
+#include "SendReceiveTestBase.h"
+
+struct FConcertBatchReplicationEvent;
+class IConcertClientReplicationManager;
+class UTestReflectionObject;
+
+namespace UE::ConcertSyncServer::Replication
+{
+	class IConcertServerReplicationManager;
+}
+
+namespace UE::ConcertSyncTests::Replication
+{
+	class FConcertClientReplicationBridgeMock;
+	
+	/**
+	 * Creates a server and connects a sender and receiver client.
+	 * The clients will be "sending" an UTestReflectionObject among each other.
+	 */
+	class FSendReceiveObjectTestBase : public FSendReceiveTestBase
+	{
+	public:
+
+		FSendReceiveObjectTestBase(const FString& InName, const bool bInComplexTask)
+			: FSendReceiveTestBase(InName, bInComplexTask)
+		{}
+
+	protected:
+		
+		using FReceiveReplicationEventSignature = void(const FConcertSessionContext& Context, const FConcertBatchReplicationEvent& Event);
+		
+		UTestReflectionObject* TestObject = nullptr;
+
+		//~ Begin FSendReceiveTestBase Interface
+		virtual ConcertSyncClient::Replication::FJoinReplicatedSessionArgs CreateSenderArgs() override;
+		virtual ConcertSyncClient::Replication::FJoinReplicatedSessionArgs CreateReceiverArgs() override;
+		virtual void SetUpClientAndServer() override;
+		
+		/**
+		 * bHasServerReceivedData and bHasClientReceivedData are reset to false prior to sending.
+		 * Sets test values on TestObject and sends it to the receiver.
+		 * If the data arrived, bHasServerReceivedData and bHasClientReceivedData are true.
+		 */
+		void SimulateSenderToReceiver(
+			TFunctionRef<FReceiveReplicationEventSignature> OnServerReceive = [](auto&, auto&){},
+			TFunctionRef<FReceiveReplicationEventSignature> OnReceiverClientReceive = [](auto&, auto&){}
+			);
+		//~ End FSendReceiveTestBase Interface
+
+		void SetTestValues(UTestReflectionObject& Object);
+		void SetDifferentValues(UTestReflectionObject& Object);
+		void TestEqualTestValues(UTestReflectionObject& Object, FAutomationTestBase& Test);
+		
+	private:
+		
+		const float SentFloat = 42.f;
+		const FVector SentVector = { 21.f, 84.f, -1.f };
+		const FGuid SenderStreamId = FGuid::NewGuid();
+	};
+}
diff --git a/Engine/Plugins/Developer/Concert/ConcertSync/ConcertSyncTest/Source/ConcertSyncTest/Private/Replication/Util/SendReceiveTestBase.cpp b/Engine/Plugins/Developer/Concert/ConcertSync/ConcertSyncTest/Source/ConcertSyncTest/Private/Replication/Util/SendReceiveTestBase.cpp
new file mode 100644
--- /dev/null
+++ b/Engine/Plugins/Developer/Concert/ConcertSync/ConcertSyncTest/Source/ConcertSyncTest/Private/Replication/Util/SendReceiveTestBase.cpp
@@ -0,0 +1,99 @@
+// Copyright Epic Games, Inc. All Rights Reserved.
+
+#include "SendReceiveTestBase.h"
+
+#include "ConcertClientReplicationBridgeMock.h"
+#include "Replication/IConcertClientReplicationManager.h"
+#include "Replication/Messages/ConcertReplicationEvents.h"
+#include "Replication/Messages/ConcertReplicationHandshakeMessages.h"
+#include "Replication/ReplicationTestInterface.h"
+#include "Util/ClientServerCommunicationTest.h"
+
+#include "Misc/AutomationTest.h"
+#include "Misc/Paths.h"
+#include "HAL/FileManager.h"
+#include "UObject/UObjectGlobals.h"
+#include "UObject/Package.h"
+
+namespace UE::ConcertSyncTests::Replication
+{
+	FSendReceiveTestBase::FSendReceiveTestBase(const FString& InName, const bool bInComplexTask)
+			: FConcertClientServerCommunicationTest(InName, bInComplexTask)
+	{}
+
+	void FSendReceiveTestBase::SetUpClientAndServer()
+	{
+		// Server
+		InitServer();
+		ServerSession = GetServerSessionMock();
+		ServerReplicationManager = ConcertSyncServer::TestInterface::CreateServerReplicationManager(ServerSession.ToSharedRef());
+		// Client Receiver
+		Client_Receiver = &ConnectClient();
+		BridgeMock_Receiver = MakeShared<FConcertClientReplicationBridgeMock>();
+		ClientReplicationManager_Receiver = ConcertSyncClient::TestInterface::CreateClientReplicationManager(Client_Receiver->ClientSessionMock, BridgeMock_Receiver.Get());
+		// Client Sender
+		Client_Sender = &ConnectClient();
+		BridgeMock_Sender = MakeShared<FConcertClientReplicationBridgeMock>();
+		ClientReplicationManager_Sender = ConcertSyncClient::TestInterface::CreateClientReplicationManager(Client_Sender->ClientSessionMock, BridgeMock_Sender.Get());
+		
+		// 1.1 Sender offers to send all UTestReflectionObject properties
+		{
+			ClientReplicationManager_Sender->JoinReplicationSession(CreateSenderArgs())
+				.Next([this](const ConcertSyncClient::Replication::FJoinReplicatedSessionResult& Result){ TestTrue(TEXT("Sender joined"), Result.ErrorCode == EJoinReplicationErrorCode::Success); });
+		}
+		// 1.2 Receiver accepts all properties
+		{
+			ClientReplicationManager_Receiver->JoinReplicationSession(CreateReceiverArgs())
+				.Next([this](const ConcertSyncClient::Replication::FJoinReplicatedSessionResult& Result){ TestTrue(TEXT("Receiver joined"), Result.ErrorCode == EJoinReplicationErrorCode::Success); });
+		}
+		
+		// 1.3 Test basic assumptions: only sender sends, only receiver receives.
+		auto TestReplicationData_Server = [&](const FConcertSessionContext& Context, const FConcertBatchReplicationEvent& Event)
+		{
+			TestEqual(TEXT("Only sender is supposed to send data!"), Context.SourceEndpointId, Client_Sender->ClientSessionMock->GetSessionClientEndpointId());
+		};
+		auto TestReplicationData_Client_Sender = [&](const FConcertSessionContext& Context, const FConcertBatchReplicationEvent& Event)
+        {
+        	// There was a bug where client would receive its own data...
+			AddError(TEXT("Sender was not supposed to receive data!"));
+        };
+		
+		ServerSession->RegisterCustomEventHandler<FConcertBatchReplicationEvent>(TestReplicationData_Server);
+		Client_Sender->ClientSessionMock->RegisterCustomEventHandler<FConcertBatchReplicationEvent>(TestReplicationData_Client_Sender);
+	}
+	
+	void FSendReceiveTestBase::SimulateSenderToReceiver()
+	{
+		// 1. Sender > Server
+		TickClient(Client_Sender);
+		// 2. Forward from server to receiver
+		TickServer();
+		// 3. Receive from server
+		TickClient(Client_Receiver);
+	}
+	
+	void FSendReceiveTestBase::TickClient(FClientInfo* Client)
+	{
+		Client->ClientSessionMock->OnTick().Broadcast(*Client->ClientSessionMock, FakeDeltaTime);
+	}
+
+	void FSendReceiveTestBase::TickServer()
+	{
+		ServerSession->OnTick().Broadcast(*ServerSession, FakeDeltaTime);
+	}
+
+	void FSendReceiveTestBase::CleanUpTest(FAutomationTestBase* AutomationTestBase)
+	{
+		if (AutomationTestBase == this)
+		{
+			ClientReplicationManager_Sender.Reset();
+			BridgeMock_Sender.Reset();
+			ClientReplicationManager_Receiver.Reset();
+			BridgeMock_Receiver.Reset();
+			ServerReplicationManager.Reset();
+			ServerSession.Reset();
+		}
+
+		FConcertClientServerCommunicationTest::CleanUpTest(AutomationTestBase);
+	}
+}
diff --git a/Engine/Plugins/Developer/Concert/ConcertSync/ConcertSyncTest/Source/ConcertSyncTest/Private/Replication/Util/SendReceiveTestBase.h b/Engine/Plugins/Developer/Concert/ConcertSync/ConcertSyncTest/Source/ConcertSyncTest/Private/Replication/Util/SendReceiveTestBase.h
new file mode 100644
--- /dev/null
+++ b/Engine/Plugins/Developer/Concert/ConcertSync/ConcertSyncTest/Source/ConcertSyncTest/Private/Replication/Util/SendReceiveTestBase.h
@@ -0,0 +1,59 @@
+// Copyright Epic Games, Inc. All Rights Reserved.
+
+#pragma once
+
+#include "Replication/IConcertClientReplicationManager.h"
+#include "Util/ClientServerCommunicationTest.h"
+
+class IConcertClientReplicationManager;
+class UTestReflectionObject;
+
+namespace UE::ConcertSyncClient::Replication
+{
+	struct FJoinReplicatedSessionArgs;
+}
+namespace UE::ConcertSyncServer::Replication
+{
+	class IConcertServerReplicationManager;
+}
+
+namespace UE::ConcertSyncTests::Replication
+{
+	class FConcertClientReplicationBridgeMock;
+	
+	/** Creates a server, connects a sender and receiver client, and completes a handshake for them. */
+	class FSendReceiveTestBase : public FConcertClientServerCommunicationTest
+	{
+	public:
+	
+		FSendReceiveTestBase(const FString& InName, const bool bInComplexTask);
+
+	protected:
+		
+		float FakeDeltaTime = 1.f / 60.f;
+		
+		TSharedPtr<FConcertServerSessionMock> ServerSession;
+		TSharedPtr<ConcertSyncServer::Replication::IConcertServerReplicationManager> ServerReplicationManager;
+		
+		FClientInfo* Client_Receiver = nullptr;
+		TSharedPtr<FConcertClientReplicationBridgeMock> BridgeMock_Receiver;
+		TSharedPtr<IConcertClientReplicationManager> ClientReplicationManager_Receiver;
+
+		FClientInfo* Client_Sender  = nullptr;
+		TSharedPtr<FConcertClientReplicationBridgeMock> BridgeMock_Sender;
+		TSharedPtr<IConcertClientReplicationManager> ClientReplicationManager_Sender;
+
+		virtual ConcertSyncClient::Replication::FJoinReplicatedSessionArgs CreateSenderArgs() = 0;
+		virtual ConcertSyncClient::Replication::FJoinReplicatedSessionArgs CreateReceiverArgs() = 0;
+
+		virtual void SetUpClientAndServer();
+		void SimulateSenderToReceiver();
+		
+		void TickClient(FClientInfo* Client);
+		void TickServer();
+		
+	private:
+		
+		virtual void CleanUpTest(FAutomationTestBase* AutomationTestBase) override;
+	};
+}
diff --git a/Engine/Plugins/Developer/Concert/ConcertSync/ConcertSyncTest/Source/ConcertSyncTest/Private/Util/ClientServerCommunicationTest.h b/Engine/Plugins/Developer/Concert/ConcertSync/ConcertSyncTest/Source/ConcertSyncTest/Private/Util/ClientServerCommunicationTest.h
new file mode 100644
--- /dev/null
+++ b/Engine/Plugins/Developer/Concert/ConcertSync/ConcertSyncTest/Source/ConcertSyncTest/Private/Util/ClientServerCommunicationTest.h
@@ -0,0 +1,459 @@
+// Copyright Epic Games, Inc. All Rights Reserved.
+
+#pragma once
+
+#include "ConcertMessages.h"
+#include "IConcertSession.h"
+#include "Misc/AutomationTest.h"
+#include "Misc/Paths.h"
+#include "Scratchpad/ConcertScratchpad.h"
+#include "Scratchpad/ConcertScratchpadPtr.h"
+#include "UObject/StructOnScope.h"
+
+namespace UE::ConcertSyncTests
+{
+	// Utility functions used to detect when a non-mocked function is called, so that we can mock it properly when required.
+	template<typename T> T NotMocked(T Ret) { check(false); return Ret; }
+	template<typename T> T NotMocked()      { check(false); return T(); }
+
+	inline FString GetTestSessionRootPath()
+	{
+		return FPaths::ProjectIntermediateDir() / TEXT("ConcertDataStoreTest");
+	}
+	
+	/** Implements a not-working IConcertServerSession. It must be further overridden to implement just what is required by the tests */
+	class FConcertServerSessionBaseMock : public IConcertServerSession
+	{
+	public:
+		FConcertServerSessionBaseMock() 
+			: Id(FGuid::NewGuid())
+			, Name("FConcertServerSessionBaseMock")
+		{ }
+
+		// IConcertSession Begin.
+		virtual const FGuid& GetId() const override																								{ return Id; }
+		virtual const FString& GetName() const override																							{ return NotMocked<const FString&>(Name); }
+		virtual void SetName(const FString&) override																							{ return NotMocked<void>(); }
+		virtual FMessageAddress GetClientAddress(const FGuid& ClientEndpointId) const override													{ return NotMocked<FMessageAddress>(); }
+		virtual const FConcertSessionInfo& GetSessionInfo() const override																		{ return NotMocked<const FConcertSessionInfo&>(SessionInfo); }
+		virtual FString GetSessionWorkingDirectory() const override																				{ return GetTestSessionRootPath() / Id.ToString(); }
+		virtual TArray<FGuid> GetSessionClientEndpointIds() const override																		{ return NotMocked<TArray<FGuid>>(); }
+		virtual TArray<FConcertSessionClientInfo> GetSessionClients() const override															{ return NotMocked<TArray<FConcertSessionClientInfo>>(); }
+		virtual bool FindSessionClient(const FGuid&, FConcertSessionClientInfo&) const override													{ return NotMocked<bool>(); }
+		virtual void Startup() override																											{ return NotMocked<void>(); }
+		virtual void Shutdown() override																										{ return NotMocked<void>(); };
+		virtual FConcertScratchpadRef GetScratchpad() const override																			{ return NotMocked<FConcertScratchpadRef>(); }
+		virtual FConcertScratchpadPtr GetClientScratchpad(const FGuid&) const override															{ return NotMocked<FConcertScratchpadRef>(); }
+		virtual FDelegateHandle InternalRegisterCustomEventHandler(const FName&, const TSharedRef<IConcertSessionCustomEventHandler>&) override { return NotMocked<FDelegateHandle>(); }
+		virtual void InternalUnregisterCustomEventHandler(const FName& EventMessageType, const FDelegateHandle EventHandle) override			{ return NotMocked<void>(); }
+		virtual void InternalUnregisterCustomEventHandler(const FName& EventMessageType, const void* EventHandler) override						{ return NotMocked<void>(); }
+		virtual void InternalClearCustomEventHandler(const FName& EventMessageType) override													{ return NotMocked<void>(); }
+		virtual void InternalSendCustomEvent(const UScriptStruct*, const void*, const TArray<FGuid>&, EConcertMessageFlags, TOptional<FConcertSequencedCustomEvent> InSequenceId={}) override			{ return NotMocked<void>(); }
+		virtual void InternalRegisterCustomRequestHandler(const FName&, const TSharedRef<IConcertSessionCustomRequestHandler>&) override		{ return NotMocked<void>(); }
+		virtual void InternalUnregisterCustomRequestHandler(const FName&) override																{ return NotMocked<void>(); }
+		virtual void InternalSendCustomRequest(const UScriptStruct*, const void*, const FGuid&, const TSharedRef<IConcertSessionCustomResponseHandler>&) override { NotMocked<void>(); }
+		// IConcertSession End.
+
+		// IConcertServerSession Begin
+		virtual FOnConcertServerSessionTick& OnTick() override                          { return NotMocked<FOnConcertServerSessionTick&>(Tick); }
+		virtual FOnConcertServerSessionClientChanged& OnSessionClientChanged() override { return NotMocked<FOnConcertServerSessionClientChanged&>(ConnectionChanged); }
+		virtual FOnConcertMessageAcknowledgementReceivedFromLocalEndpoint& OnConcertMessageAcknowledgementReceived() override { return NotMocked<FOnConcertMessageAcknowledgementReceivedFromLocalEndpoint&>(AckReceived); }
+		// IConcertServerSession End
+
+	protected:
+		FGuid Id;
+
+		// Those below are unused mocked data member, but required to get the code compiling.
+		FString Name;
+		FConcertSessionInfo SessionInfo;
+		FOnConcertServerSessionTick Tick;
+		FOnConcertServerSessionClientChanged ConnectionChanged;
+		FOnConcertMessageAcknowledgementReceivedFromLocalEndpoint AckReceived;
+	};
+
+	/** Implements a not-working IConcertClientSession. It must be further overridden to implement just what is required by the tests */
+	class FConcertClientSessionBaseMock : public IConcertClientSession
+	{
+	public:
+		FConcertClientSessionBaseMock() 
+			: Id(FGuid::NewGuid()) 
+			, Name("FConcertClientSessionBaseMock") 
+		{ }
+
+		// IConcertSession Begin.
+		virtual const FGuid& GetId() const override																								{ return Id; }
+		virtual const FString& GetName() const override																							{ return NotMocked<const FString&>(Name); }
+		virtual const FConcertSessionInfo& GetSessionInfo() const override																		{ return NotMocked<const FConcertSessionInfo&>(SessionInfo); }
+		virtual FString GetSessionWorkingDirectory() const override																				{ return GetTestSessionRootPath() / Id.ToString(); }
+		virtual TArray<FGuid> GetSessionClientEndpointIds() const override																		{ return NotMocked<TArray<FGuid>>(); }
+		virtual TArray<FConcertSessionClientInfo> GetSessionClients() const override															{ return NotMocked<TArray<FConcertSessionClientInfo>>(); }
+		virtual bool FindSessionClient(const FGuid&, FConcertSessionClientInfo&) const override													{ return NotMocked<bool>(); }
+		virtual void Startup() override																											{ return NotMocked<void>(); }
+		virtual void Shutdown() override																										{ return NotMocked<void>(); };
+		virtual FConcertScratchpadRef GetScratchpad() const override																			{ return NotMocked<FConcertScratchpadRef>(); }
+		virtual FConcertScratchpadPtr GetClientScratchpad(const FGuid&) const override															{ return NotMocked<FConcertScratchpadRef>(); }
+		virtual FDelegateHandle InternalRegisterCustomEventHandler(const FName&, const TSharedRef<IConcertSessionCustomEventHandler>&) override { return NotMocked<FDelegateHandle>(); }
+		virtual void InternalUnregisterCustomEventHandler(const FName& EventMessageType, const FDelegateHandle EventHandle) override			{ return NotMocked<void>(); }
+		virtual void InternalUnregisterCustomEventHandler(const FName& EventMessageType, const void* EventHandler) override						{ return NotMocked<void>(); }
+		virtual void InternalClearCustomEventHandler(const FName& EventMessageType) override													{ return NotMocked<void>(); }
+		virtual void InternalSendCustomEvent(const UScriptStruct*, const void*, const TArray<FGuid>&, EConcertMessageFlags, TOptional<FConcertSequencedCustomEvent> InSequenceId={}) override			{ return NotMocked<void>(); }
+		virtual void InternalRegisterCustomRequestHandler(const FName&, const TSharedRef<IConcertSessionCustomRequestHandler>&) override		{ return NotMocked<void>(); }
+		virtual void InternalUnregisterCustomRequestHandler(const FName&) override																{ return NotMocked<void>(); }
+		virtual void InternalSendCustomRequest(const UScriptStruct*, const void*, const FGuid&, const TSharedRef<IConcertSessionCustomResponseHandler>&) override { NotMocked<void>(); }
+		// IConcertSession End.
+
+		virtual FConcertSequencedCustomEventManager& GetSequencedEventManager() override
+		{
+			return NotMocked<FConcertSequencedCustomEventManager&>(CustomEventSequenceManager);
+		}
+		// IConcertClientSession Begin
+		virtual EConcertConnectionStatus GetConnectionStatus() const override            { return NotMocked<EConcertConnectionStatus>(EConcertConnectionStatus::Connected); }
+		virtual FGuid GetSessionClientEndpointId() const override                        { return NotMocked<FGuid>(); }
+		virtual FGuid GetSessionServerEndpointId() const override                        { return NotMocked<FGuid>(); }
+		virtual const FConcertClientInfo& GetLocalClientInfo() const override            { return NotMocked<const FConcertClientInfo&>(ClientInfo); }
+		virtual void UpdateLocalClientInfo(const FConcertClientInfoUpdate&) override     { return NotMocked<void>(); }
+		virtual void Connect() override                                                  { return NotMocked<void>(); }
+		virtual void Disconnect() override                                               { return NotMocked<void>(); }
+
+		virtual EConcertSendReceiveState GetSendReceiveState() const override
+		{
+			return NotMocked<EConcertSendReceiveState>(EConcertSendReceiveState::Default);
+		}
+
+		virtual void SetSendReceiveState(EConcertSendReceiveState InSendReceiveState) override
+		{
+			return NotMocked<void>();
+		}
+
+		virtual FOnConcertClientSessionTick& OnTick() override                           { return NotMocked<FOnConcertClientSessionTick&>(Tick); }
+		virtual FOnConcertClientSessionConnectionChanged& OnConnectionChanged() override { return NotMocked<FOnConcertClientSessionConnectionChanged&>(ConnectionChanged); }
+		virtual FOnConcertClientSessionClientChanged& OnSessionClientChanged() override  { return NotMocked<FOnConcertClientSessionClientChanged&>(ClientChanged); }
+		virtual FOnConcertSessionRenamed& OnSessionRenamed() override                    { return NotMocked<FOnConcertSessionRenamed&>(SessionRenamed); }
+		// IConcertClientSession End
+
+		virtual void HandleCustomEvent(const UScriptStruct* EventType, const void* EventData) = 0;
+
+	protected:
+		FGuid Id;
+
+		// Those below are unused mocked data member, but required to get the code compiling.
+		FString Name;
+		FConcertSessionInfo SessionInfo;
+		FOnConcertClientSessionTick Tick;
+		FOnConcertClientSessionConnectionChanged ConnectionChanged;
+		FOnConcertClientSessionClientChanged ClientChanged;
+		FOnConcertSessionRenamed SessionRenamed;
+		FConcertSequencedCustomEventManager CustomEventSequenceManager;
+		FConcertClientInfo ClientInfo;
+	};
+
+	/** Specializes the base concert server session to act as a fake server session. */
+	class FConcertServerSessionMock : public FConcertServerSessionBaseMock
+	{
+	public:
+		
+		virtual void InternalSendCustomEvent(const UScriptStruct* EventType, const void* EventData, const TArray<FGuid>& TargetEndpointIds, EConcertMessageFlags, TOptional<FConcertSequencedCustomEvent> InSequenceId={}) override
+		{
+			for (const FGuid& TargetEndPointId : TargetEndpointIds)
+			{
+				int i = 0;
+				for (const FGuid& ClientEndPointId : ClientEndpoints)
+				{
+					if (TargetEndPointId == ClientEndPointId)
+					{
+						// Dispatch the event on the client immediately.
+						ClientSessions[i]->HandleCustomEvent(EventType, EventData);
+					}
+					++i;
+				}
+			}
+		}
+
+		virtual FDelegateHandle InternalRegisterCustomEventHandler(const FName& EventMessageType, const TSharedRef<IConcertSessionCustomEventHandler>& Handler) override
+		{
+			CustomEventHandlers.FindOrAdd(EventMessageType).Add(Handler);
+			return Handler->GetHandle();
+		}
+
+		virtual void InternalUnregisterCustomEventHandler(const FName& EventMessageType, const FDelegateHandle EventHandle) override
+		{
+			TArray<TSharedPtr<IConcertSessionCustomEventHandler>>* HandlerArrayPtr = CustomEventHandlers.Find(EventMessageType);
+			if (HandlerArrayPtr)
+			{
+				for (auto It = HandlerArrayPtr->CreateIterator(); It; ++It)
+				{
+					if ((*It)->GetHandle() == EventHandle)
+					{
+						It.RemoveCurrent();
+						break;
+					}
+				}
+			}
+		}
+
+		virtual void InternalUnregisterCustomEventHandler(const FName& EventMessageType, const void* EventHandler) override
+		{
+			TArray<TSharedPtr<IConcertSessionCustomEventHandler>>* HandlerArrayPtr = CustomEventHandlers.Find(EventMessageType);
+			if (HandlerArrayPtr)
+			{
+				for (auto It = HandlerArrayPtr->CreateIterator(); It; ++It)
+				{
+					if ((*It)->HasSameObject(EventHandler))
+					{
+						It.RemoveCurrent();
+						break;
+					}
+				}
+			}
+		}
+
+		virtual void InternalRegisterCustomRequestHandler(const FName& RequestMessageType, const TSharedRef<IConcertSessionCustomRequestHandler>& Handler) override
+		{
+			CustomRequestHandlers.Add(RequestMessageType, Handler);
+		}
+
+		virtual void InternalUnregisterCustomRequestHandler(const FName& RequestMessageType) override
+		{
+			CustomRequestHandlers.Remove(RequestMessageType);
+		}
+
+		virtual TArray<FGuid> GetSessionClientEndpointIds() const override
+		{
+			return ClientEndpoints;
+		}
+
+		// Called by the tests to connect a client to the server.
+		void ConnectClient(const FGuid& ClientEndpointId, FConcertClientSessionBaseMock& ClientSession)
+		{
+			// Notify the server that a new client connected. The server data store will replicate its content on the client.
+			FConcertSessionClientInfo Info;
+			Info.ClientEndpointId = ClientEndpointId;
+			ClientEndpoints.Add(ClientEndpointId);
+			ClientSessions.Add(&ClientSession);
+			ConnectionChanged.Broadcast(*this, EConcertClientStatus::Connected, Info);
+		}
+
+		void DispatchEvent(const FGuid& RequesterEndpointId, const UScriptStruct* EventType, const void* EventData, const TArray<FGuid>& DestinationEndpointIds, EConcertMessageFlags Flags, TOptional<FConcertSequencedCustomEvent> InSequenceId)
+		{
+			// { 0, 0, 0, 0 } is for the server
+			if (DestinationEndpointIds.Contains(FGuid{ 0, 0, 0, 0 }))
+			{
+				if (const TArray<TSharedPtr<IConcertSessionCustomEventHandler>>* Handlers = CustomEventHandlers.Find(EventType->GetFName()))
+				{
+					FConcertSessionContext DummyContext;
+					DummyContext.SourceEndpointId = RequesterEndpointId;
+					for (const TSharedPtr<IConcertSessionCustomEventHandler>& Handler : *Handlers)
+					{
+						Handler->HandleEvent(DummyContext, EventData);
+					}
+				}
+			}
+			
+			InternalSendCustomEvent(EventType, EventData, DestinationEndpointIds, Flags, InSequenceId);
+		}
+
+		// Called by the FConcertClientSessionMock to dispatch a request
+		void DispatchRequest(const FGuid& RequesterEndpointId, const UScriptStruct* RequestType, const void* RequestData, const TSharedRef<IConcertSessionCustomResponseHandler>& ResponseHandler)
+		{
+			if (TSharedPtr<IConcertSessionCustomRequestHandler>* RequestHandler = CustomRequestHandlers.Find(RequestType->GetFName()))
+			{
+				// Set up who's sending the request.
+				FConcertSessionContext Context;
+				Context.SourceEndpointId = RequesterEndpointId;
+
+				// Dispatch the request
+				FStructOnScope ResponsePayload((*RequestHandler)->GetResponseType());
+				EConcertSessionResponseCode Result = (*RequestHandler)->HandleRequest(Context, RequestData, ResponsePayload.GetStructMemory());
+				if (Result == EConcertSessionResponseCode::Success || Result == EConcertSessionResponseCode::Failed)
+				{
+					// Dispatch the response.
+					ResponseHandler->HandleResponse(ResponsePayload.GetStructMemory());
+				}
+				else
+				{
+					check(false); // The test suite is not expected to fire any other result than Success or Failed.
+				}
+			}
+		}
+
+		virtual bool FindSessionClient(const FGuid& EndpointId, FConcertSessionClientInfo& OutSessionClientInfo) const override
+		{
+			// Need to mock this in case test fails the internal server logic may log the client name which made an invalid request.
+			return true;
+		}
+		
+		virtual FOnConcertServerSessionClientChanged& OnSessionClientChanged() override
+		{
+			return ConnectionChanged;
+		}
+
+		virtual FOnConcertServerSessionTick& OnTick() override { return OnTickDelegate; }
+
+	private:
+
+		/** Tick delegate for tests that need it */
+		FOnConcertServerSessionTick OnTickDelegate;
+		
+		/** Map of session custom event handlers */
+		TMap<FName, TArray<TSharedPtr<IConcertSessionCustomEventHandler>>> CustomEventHandlers;
+		/** Map of session custom request handlers */
+		TMap<FName, TSharedPtr<IConcertSessionCustomRequestHandler>> CustomRequestHandlers;
+
+		/** Connected client endpoints. */
+		TArray<FGuid> ClientEndpoints;
+
+		/** Connected clients sessions. */
+		TArray<FConcertClientSessionBaseMock*> ClientSessions;
+	};
+
+	/** Specializes the base concert client session to act as a fake client session. */
+	class FConcertClientSessionMock : public FConcertClientSessionBaseMock
+	{
+	public:
+		FConcertClientSessionMock(const FGuid& ClientEndpointId, FConcertServerSessionMock& Server)
+			: ServerMock(Server)
+			, EndpointId(ClientEndpointId)
+		{
+		}
+
+		virtual FDelegateHandle InternalRegisterCustomEventHandler(const FName& EventMessageType, const TSharedRef<IConcertSessionCustomEventHandler>& Handler) override
+		{
+			CustomEventHandlers.FindOrAdd(EventMessageType).Add(Handler);
+			return Handler->GetHandle();
+		}
+
+		virtual void InternalUnregisterCustomEventHandler(const FName& EventMessageType, const FDelegateHandle EventHandle) override
+		{
+			TArray<TSharedPtr<IConcertSessionCustomEventHandler>>* HandlerArrayPtr = CustomEventHandlers.Find(EventMessageType);
+			if (HandlerArrayPtr)
+			{
+				for (auto It = HandlerArrayPtr->CreateIterator(); It; ++It)
+				{
+					if ((*It)->GetHandle() == EventHandle)
+					{
+						It.RemoveCurrent();
+						break;
+					}
+				}
+			}
+		}
+
+		virtual void InternalUnregisterCustomEventHandler(const FName& EventMessageType, const void* EventHandler) override
+		{
+			TArray<TSharedPtr<IConcertSessionCustomEventHandler>>* HandlerArrayPtr = CustomEventHandlers.Find(EventMessageType);
+			if (HandlerArrayPtr)
+			{
+				HandlerArrayPtr->RemoveAllSwap([EventHandler](const TSharedPtr<IConcertSessionCustomEventHandler>& InHandler) -> bool
+				{
+					return InHandler->HasSameObject(EventHandler);
+				});
+			}
+		}
+
+		virtual void InternalClearCustomEventHandler(const FName& EventMessageType) override
+		{
+			CustomEventHandlers.Remove(EventMessageType);
+		}
+
+		virtual void InternalSendCustomEvent(const UScriptStruct* EventType, const void* EventData, const TArray<FGuid>& DestinationEndpointIds, EConcertMessageFlags Flags, TOptional<FConcertSequencedCustomEvent> InSequenceId) override
+		{
+			ServerMock.DispatchEvent(EndpointId, EventType, EventData, DestinationEndpointIds, Flags, InSequenceId);
+		}
+
+		virtual void InternalSendCustomRequest(const UScriptStruct* RequestType, const void* RequestData, const FGuid& DestinationEndpointId, const TSharedRef<IConcertSessionCustomResponseHandler>& Handler) override
+		{
+			// Directly dispatch to the server session.
+			ServerMock.DispatchRequest(EndpointId, RequestType, RequestData, Handler);
+		}
+
+		virtual void HandleCustomEvent(const UScriptStruct* EventType, const void* EventData) override
+		{
+			TArray<TSharedPtr<IConcertSessionCustomEventHandler>>* HandlerArrayPtr = CustomEventHandlers.Find(EventType->GetFName());
+			if (HandlerArrayPtr)
+			{
+				FConcertSessionContext DummyContext;
+				for (const TSharedPtr<IConcertSessionCustomEventHandler>& Handler : *HandlerArrayPtr)
+				{
+					Handler->HandleEvent(DummyContext, EventData);
+				}
+			}
+		}
+
+		virtual FGuid GetSessionClientEndpointId() const override
+		{
+			return EndpointId;
+		}
+
+		virtual FGuid GetSessionServerEndpointId() const override
+		{
+			return FGuid(0, 0, 0, 0);
+		}
+		
+		
+		virtual FOnConcertClientSessionTick& OnTick() override { return OnTickDelegate; }
+
+	private:
+		FOnConcertClientSessionTick OnTickDelegate;
+		FConcertServerSessionMock& ServerMock;
+		FGuid EndpointId;
+		TMap<FName, TArray<TSharedPtr<IConcertSessionCustomEventHandler>>> CustomEventHandlers;
+	};
+
+	/** Base class to perform Concert data store client/server tests. */
+	class FConcertClientServerCommunicationTest : public FAutomationTestBase
+	{
+	public:
+		struct FClientInfo
+		{
+			FClientInfo(const FGuid& ClientEndPointId, FConcertServerSessionMock& Server)
+				: ClientSessionMock(MakeShared<FConcertClientSessionMock>(ClientEndPointId, Server))
+			{}
+
+			const TSharedRef<FConcertClientSessionBaseMock> ClientSessionMock;
+		};
+
+		FConcertClientServerCommunicationTest(const FString& InName, const bool bInComplexTask)
+			: FAutomationTestBase(InName, bInComplexTask)
+		{
+			FAutomationTestFramework& Framework = FAutomationTestFramework::Get();
+			Framework.OnTestEndEvent.AddRaw(this, &FConcertClientServerCommunicationTest::CleanUpTest);
+		}
+		
+		virtual ~FConcertClientServerCommunicationTest() override
+		{
+			FAutomationTestFramework& Framework = FAutomationTestFramework::Get();
+			Framework.OnTestEndEvent.RemoveAll(this);
+		}
+
+		FClientInfo& ConnectClient()
+		{
+			const FGuid ClientEndpointId(0, 0, 0, Clients.Num() + 1); // {0, 0, 0, 0} is used by the server.
+			Clients.Add(MakeUnique<FClientInfo>(ClientEndpointId, *ServerSessionMock));
+			ServerSessionMock->ConnectClient(ClientEndpointId, *(Clients.Last()->ClientSessionMock));
+			return *Clients.Last();
+		}
+
+		void InitServer()
+		{
+			// Reset everything to be able to rerun the tests. The test framework doesn't destruct/reconstruct this object
+			// at every run, so just ensure we start with a clean state.
+			Clients.Empty();
+			ServerSessionMock = MakeShared<FConcertServerSessionMock>();
+		}
+
+		const TSharedPtr<FConcertServerSessionMock>& GetServerSessionMock() const { return ServerSessionMock; }
+
+		virtual void CleanUpTest(FAutomationTestBase* AutomationTestBase)
+		{
+			if (AutomationTestBase == this)
+			{
+				Clients.Reset();
+				ServerSessionMock.Reset();
+			}
+		}
+		
+	private:
+		TSharedPtr<FConcertServerSessionMock> ServerSessionMock;
+		TArray<TUniquePtr<FClientInfo>> Clients;
+	};
+}
\ No newline at end of file
diff --git a/Engine/Plugins/Developer/Concert/ConcertSync/ConcertSyncTest/Source/ConcertSyncTest/Private/Util/ConcertClientReplicationBridgeMock.h b/Engine/Plugins/Developer/Concert/ConcertSync/ConcertSyncTest/Source/ConcertSyncTest/Private/Util/ConcertClientReplicationBridgeMock.h
deleted file mode 100644
diff --git a/Engine/Plugins/Developer/Concert/ConcertSync/ConcertSyncTest/Source/ConcertSyncTest/Private/Util/ConcertMocks.h b/Engine/Plugins/Developer/Concert/ConcertSync/ConcertSyncTest/Source/ConcertSyncTest/Private/Util/ConcertMocks.h
deleted file mode 100644
