Clients must now request authority before sending.
MultiUser will automatically request authority for all objects when connecting.

diff --git a/Engine/Plugins/Developer/Concert/ConcertApp/MultiUserClient/Source/MultiUserClient/Private/MultiUserClientModule.cpp b/Engine/Plugins/Developer/Concert/ConcertApp/MultiUserClient/Source/MultiUserClient/Private/MultiUserClientModule.cpp
--- a/Engine/Plugins/Developer/Concert/ConcertApp/MultiUserClient/Source/MultiUserClient/Private/MultiUserClientModule.cpp
+++ b/Engine/Plugins/Developer/Concert/ConcertApp/MultiUserClient/Source/MultiUserClient/Private/MultiUserClientModule.cpp
@@ -2,19 +2,20 @@
 
 #include "IMultiUserClientModule.h"
 
-#include "IConcertSyncClientModule.h"
-#include "IConcertModule.h"
-#include "IConcertClient.h"
 #include "ConcertClientSettings.h"
-#include "IConcertSession.h"
-#include "IConcertSyncClient.h"
 #include "ConcertLogGlobal.h"
 #include "ConcertFrontendStyle.h"
 #include "ConcertSyncSettings.h"
 #include "ConcertWorkspaceData.h"
 #include "ConcertWorkspaceUI.h"
 #include "ConcertUtil.h"
+#include "IConcertClient.h"
+#include "IConcertModule.h"
+#include "IConcertSession.h"
+#include "IConcertSyncClient.h"
+#include "IConcertSyncClientModule.h"
 #include "MultiUserClientUtils.h"
+#include "Replication/MultiUserReplicationManager.h"
 
 #include "Misc/App.h"
 #include "Misc/AsyncTaskNotification.h"
@@ -621,6 +622,7 @@ public:
 
 		// Boot the client instance
 		MultiUserClient->Startup(ClientConfig, EConcertSyncSessionFlags::Default_MultiUserSession);
+		ReplicationManager = MakeShared<UE::MultiUserClient::FMultiUserReplicationManager>(MultiUserClient.ToSharedRef());
 
 		// Hook UI elements in the tool bar (and setup commands).
 		RegisterUI();
@@ -1043,7 +1045,7 @@ private:
 	{
 		const TSharedRef<SDockTab> DockTab = SNew(SDockTab)
 			.TabRole(NomadTab);
-		DockTab->SetContent(SNew(SConcertBrowser, DockTab, MultiUserClient));
+		DockTab->SetContent(SNew(SConcertBrowser, DockTab, MultiUserClient.ToSharedRef(), ReplicationManager.ToSharedRef()));
 		return DockTab;
 	}
 
@@ -1386,6 +1388,8 @@ private:
 	}
 
 	TSharedPtr<IConcertSyncClient> MultiUserClient;
+	/** Interacts with the replication system on behalf of Multi-User. */
+	TSharedPtr<UE::MultiUserClient::FMultiUserReplicationManager> ReplicationManager;
 
 	/** True if the tab spawners have been registered for this module */
 	bool bHasRegisteredTabSpawners = false;
diff --git a/Engine/Plugins/Developer/Concert/ConcertApp/MultiUserClient/Source/MultiUserClient/Private/Replication/MultiUserReplicationManager.cpp b/Engine/Plugins/Developer/Concert/ConcertApp/MultiUserClient/Source/MultiUserClient/Private/Replication/MultiUserReplicationManager.cpp
new file mode 100644
--- /dev/null
+++ b/Engine/Plugins/Developer/Concert/ConcertApp/MultiUserClient/Source/MultiUserClient/Private/Replication/MultiUserReplicationManager.cpp
@@ -0,0 +1,39 @@
+// Copyright Epic Games, Inc. All Rights Reserved.
+
+#include "MultiUserReplicationManager.h"
+
+#include "IConcertSyncClient.h"
+#include "MultiUserReplicationClientProfileAsset.h"
+#include "ReplicationUtils.h"
+
+namespace UE::MultiUserClient
+{
+	FMultiUserReplicationManager::FMultiUserReplicationManager(TSharedRef<IConcertSyncClient> InClient)
+		: Client(MoveTemp(InClient))
+	{}
+
+	void FMultiUserReplicationManager::JoinReplicationSession(const UMultiUserReplicationClientProfileAsset& Asset)
+	{
+		Replication::JoinSessionForMultiUser(Client, Asset.ToJoinArgs());
+	}
+
+	void FMultiUserReplicationManager::LeaveSession()
+	{
+		if (IConcertClientReplicationManager* ReplicationManager = Client->GetReplicationManager())
+		{
+			ReplicationManager->LeaveReplicationSession();
+		}
+	}
+
+	bool FMultiUserReplicationManager::CanJoin() const
+	{
+		IConcertClientReplicationManager* ReplicationManager = Client->GetReplicationManager();
+		return ReplicationManager && ReplicationManager->CanJoin();
+	}
+
+	bool FMultiUserReplicationManager::IsConnectedToReplicationSession() const
+	{
+		IConcertClientReplicationManager* ReplicationManager = Client->GetReplicationManager();
+		return ReplicationManager && ReplicationManager->IsConnectedToReplicationSession();
+	}
+}
diff --git a/Engine/Plugins/Developer/Concert/ConcertApp/MultiUserClient/Source/MultiUserClient/Private/Replication/MultiUserReplicationManager.h b/Engine/Plugins/Developer/Concert/ConcertApp/MultiUserClient/Source/MultiUserClient/Private/Replication/MultiUserReplicationManager.h
new file mode 100644
--- /dev/null
+++ b/Engine/Plugins/Developer/Concert/ConcertApp/MultiUserClient/Source/MultiUserClient/Private/Replication/MultiUserReplicationManager.h
@@ -0,0 +1,44 @@
+// Copyright Epic Games, Inc. All Rights Reserved.
+
+#pragma once
+
+#include "Templates/SharedPointer.h"
+
+class UMultiUserReplicationClientProfileAsset;
+class IConcertSyncClient;
+
+namespace UE::MultiUserClient
+{
+	/**
+	 * Interacts with the replication system on behalf of Multi-User to execute actions specific to Multi-User workflows;
+	 * this is opposed to other uses of the replication API, e.g. users using the system in a shipped game.
+	 *
+	 * This class will
+	 *  - be used as a model for the MU control views, such as displaying the streams in the current session (TODO DP UE-193541).
+	 *  - implement auto-join behavior in response to joining a concert session (TODO DP UE-193538)
+	 *
+	 * This class implements the Fence design pattern. All knowledge Multi-User might need should be encapsulated by this class.
+	 */
+	class FMultiUserReplicationManager
+	{
+	public:
+
+		FMultiUserReplicationManager(TSharedRef<IConcertSyncClient> InClient);
+
+		/** Joins a replication session. */
+		void JoinReplicationSession(const UMultiUserReplicationClientProfileAsset& Asset);
+		/** Leaves the current replication session */
+		void LeaveSession();
+
+		/** Whether it is valid to call JoinReplicationSession. */
+		bool CanJoin() const;
+		/** Whether it is valid to call LeaveSession. */
+		bool IsConnectedToReplicationSession() const;
+
+	private:
+
+		/** Client through which the replication bridge is accessed. */
+		TSharedRef<IConcertSyncClient> Client;
+	};
+}
+
diff --git a/Engine/Plugins/Developer/Concert/ConcertApp/MultiUserClient/Source/MultiUserClient/Private/Replication/ReplicationUtils.cpp b/Engine/Plugins/Developer/Concert/ConcertApp/MultiUserClient/Source/MultiUserClient/Private/Replication/ReplicationUtils.cpp
new file mode 100644
--- /dev/null
+++ b/Engine/Plugins/Developer/Concert/ConcertApp/MultiUserClient/Source/MultiUserClient/Private/Replication/ReplicationUtils.cpp
@@ -0,0 +1,280 @@
+// Copyright Epic Games, Inc. All Rights Reserved.
+
+#include "ReplicationUtils.h"
+
+#include "ConcertLogGlobal.h"
+#include "IConcertSyncClient.h"
+#include "Replication/Messages/ConcertReplicationHandshakeMessages.h"
+#include "Replication/IConcertClientReplicationManager.h"
+
+#include "Async/Async.h"
+#include "Framework/Notifications/NotificationManager.h"
+#include "Widgets/Notifications/SNotificationList.h"
+
+#define LOCTEXT_NAMESPACE "ReplicationUtils"
+
+namespace UE::MultiUserClient::Replication::Private
+{
+	/** Wraps join request with SNotificationItem. */
+	static TPair<TFuture<ConcertSyncClient::Replication::FJoinReplicatedSessionResult>, TSharedPtr<SNotificationItem>> JoinSession(
+		IConcertClientReplicationManager& Manager,
+		ConcertSyncClient::Replication::FJoinReplicatedSessionArgs Args
+		);
+
+	/** Reuses SNotificationItem from join for authority change. */
+	static TFuture<FJoinSessionResult> RequestAuthorityAfterJoin(
+		const TSharedRef<IConcertSyncClient>& Client,
+		TFuture<ConcertSyncClient::Replication::FJoinReplicatedSessionResult>&& JoinRequestFuture,
+		ConcertSyncClient::Replication::FJoinReplicatedSessionArgs Args,
+		TSharedPtr<SNotificationItem> NotificationToReuse
+		);
+
+	/** Updates the message started by JoinSession */
+	static void HandleJoinMessage(
+		TSharedPtr<SNotificationItem> NotificationItem,
+		ConcertSyncClient::Replication::FJoinReplicatedSessionResult ErrorResult
+		);
+	/** Updates the message started by RequestChangeAuthority. */
+	static void HandleAuthorityMessage(
+		TSharedPtr<SNotificationItem> NotificationItem,
+		const int32 NumRequested,
+		const int32 NumReleased,
+		ConcertSyncClient::Replication::FAuthorityChangeResponse Response
+		);
+}
+
+namespace UE::MultiUserClient::Replication
+{
+	TFuture<FJoinSessionResult> JoinSessionForMultiUser(
+		const TSharedRef<IConcertSyncClient>& Client,
+		ConcertSyncClient::Replication::FJoinReplicatedSessionArgs Args
+		)
+	{
+		using namespace ConcertSyncClient::Replication;
+		IConcertClientReplicationManager* Manager = Client->GetReplicationManager();
+		if (!Manager)
+		{
+			return MakeFulfilledPromise<FJoinSessionResult>(FJoinSessionResult{ { EJoinReplicationErrorCode::Cancelled } }).GetFuture();
+		}
+		
+		// 1. Join the session
+		auto[JoinRequestFuture, NotificationItem] = Private::JoinSession(*Manager, Args); 
+		// 2. Request authority for all the client's objects
+		// The notification for joining is re-used to avoid spam
+		TFuture<FJoinSessionResult> RequestAuthorityFuture = Private::RequestAuthorityAfterJoin(Client, MoveTemp(JoinRequestFuture), Args, NotificationItem);
+		
+		return RequestAuthorityFuture;
+	}
+	
+	TFuture<ConcertSyncClient::Replication::FAuthorityChangeResponse> RequestAuthorityChange(
+		IConcertClientReplicationManager& Manager,
+		ConcertSyncClient::Replication::FAuthorityChangeRequest Args,
+		TSharedPtr<SNotificationItem> NotificationToReuse
+		)
+	{
+		using namespace UE::ConcertSyncClient::Replication;
+		const int32 NumRequested = Args.TakeAuthority.Num();
+		const int32 NumReleased = Args.ReleaseAuthority.Num();
+		
+		// IConcertClientReplicationManager::RequestAuthorityChange already checks for this case but we don't want to show the notification.
+		if (NumRequested == 0 && NumReleased == 0)
+		{
+			if (NotificationToReuse)
+			{
+				NotificationToReuse->ExpireAndFadeout();
+			}
+			return MakeFulfilledPromise<FAuthorityChangeResponse>(FAuthorityChangeResponse{}).GetFuture();
+		}
+		TFuture<FAuthorityChangeResponse> Future = Manager.RequestAuthorityChange(Args);
+		
+		const FText Title = LOCTEXT("AuthorityChange.Start", "Requesting Authority Change");
+		TSharedPtr<SNotificationItem> NotificationItem = MoveTemp(NotificationToReuse);
+		if (!NotificationItem)
+		{
+			FNotificationInfo Notification(Title);
+			Notification.bFireAndForget = false;
+			Notification.ExpireDuration = 6.f;
+			Notification.bUseThrobber = true;
+			NotificationItem = FSlateNotificationManager::Get().AddNotification(Notification);
+		}
+		else
+		{
+			NotificationItem->SetText(Title);
+		}
+
+		// Adding notifications should usually succeed but can theoretically fail
+		if (NotificationItem)
+		{
+			NotificationItem->SetCompletionState(SNotificationItem::CS_Pending);
+			Future = Future.Next([NotificationItem = MoveTemp(NotificationItem), NumRequested, NumReleased](FAuthorityChangeResponse Response) mutable
+			{
+				Private::HandleAuthorityMessage(MoveTemp(NotificationItem), NumRequested, NumReleased, MoveTemp(Response));
+				return Response;
+			});
+		}
+		
+		return Future;
+	}
+	
+	void LeaveSessionForMultiUser(IConcertClientReplicationManager& Manager)
+	{
+		Manager.LeaveReplicationSession();
+		
+		FNotificationInfo Notification(LOCTEXT("Left", "Left Replication Session"));
+		Notification.bFireAndForget = true;
+		Notification.ExpireDuration = 4.f;
+		FSlateNotificationManager::Get().AddNotification(Notification);
+	}
+}
+
+namespace UE::MultiUserClient::Replication::Private
+{
+	static TPair<TFuture<ConcertSyncClient::Replication::FJoinReplicatedSessionResult>, TSharedPtr<SNotificationItem>> JoinSession(
+		IConcertClientReplicationManager& Manager,
+		ConcertSyncClient::Replication::FJoinReplicatedSessionArgs Args)
+	{
+		using namespace ConcertSyncClient::Replication;
+		TFuture<FJoinReplicatedSessionResult> Future = Manager.JoinReplicationSession(MoveTemp(Args));
+		
+		FNotificationInfo Notification(LOCTEXT("Joining", "Joining Replication Session"));
+		Notification.bFireAndForget = false;
+		Notification.ExpireDuration = 6.f;
+		Notification.bUseThrobber = true;
+		TSharedPtr<SNotificationItem> NotificationItem = FSlateNotificationManager::Get().AddNotification(Notification);
+
+		// Adding notifications should usually succeed but can theoretically fail
+		if (NotificationItem)
+		{
+			NotificationItem->SetCompletionState(SNotificationItem::CS_Pending);
+			Future = Future.Next([NotificationItem](FJoinReplicatedSessionResult ErrorResult) mutable
+			{
+				HandleJoinMessage(MoveTemp(NotificationItem), ErrorResult);
+				return ErrorResult;
+			});
+		}
+		
+		return { MoveTemp(Future), NotificationItem };
+	}
+	
+	static TFuture<FJoinSessionResult> RequestAuthorityAfterJoin(
+		const TSharedRef<IConcertSyncClient>& Client,
+		TFuture<ConcertSyncClient::Replication::FJoinReplicatedSessionResult>&& JoinRequestFuture,
+		ConcertSyncClient::Replication::FJoinReplicatedSessionArgs Args,
+		TSharedPtr<SNotificationItem> NotificationToReuse
+		)
+	{
+		TPromise<FJoinSessionResult> Promise;
+		TFuture<FJoinSessionResult> Future = Promise.GetFuture();
+		JoinRequestFuture.Next([WeakClient = Client.ToWeakPtr(), NotificationToReuse = MoveTemp(NotificationToReuse), Promise = MoveTemp(Promise), Args]
+			(ConcertSyncClient::Replication::FJoinReplicatedSessionResult JoinResult) mutable
+		{
+			ON_SCOPE_EXIT
+			{
+				// Note that RequestChangeAuthority will MoveTemp this to null (if we sent an authority request)
+				if (NotificationToReuse)
+				{
+					NotificationToReuse->ExpireAndFadeout();
+				}
+			};
+
+			ConcertSyncClient::Replication::FAuthorityChangeRequest AuthorityRequest;
+			for (const FReplicationStreamDescription& Stream : Args.Streams)
+			{
+				const FSharedReplicationStreamDescription& StreamDescription = Stream.BaseDescription;
+				for (const TPair<FSoftObjectPath, FReplicatedObjectInfo>& ReplicatedObject : StreamDescription.ReplicationMap.ReplicatedObjects)
+				{
+					AuthorityRequest.TakeAuthority.FindOrAdd(ReplicatedObject.Key).StreamIds.Add(StreamDescription.Identifier);
+				}
+			}
+			
+			// Client pointer is needed because it is unsafe to keep the IConcertClientReplicationManager* reference
+			const TSharedPtr<IConcertSyncClient> ClientPin = WeakClient.Pin();
+			IConcertClientReplicationManager* ReplicationManager = ClientPin ? ClientPin->GetReplicationManager() : nullptr;
+			const bool bTimedOut = IsEngineExitRequested() || !ReplicationManager;
+			const bool bErrorJoining = JoinResult.ErrorCode != EJoinReplicationErrorCode::Success;
+			const bool bNothingToRequest = AuthorityRequest.TakeAuthority.IsEmpty();
+			// TODO DP: Here is a good place to check project settings whether to automatically request authority (but then make sure to expire NotificationToReuse).
+			if (bTimedOut || bErrorJoining || bNothingToRequest)
+			{
+				Promise.EmplaceValue(FJoinSessionResult{ MoveTemp(JoinResult) });
+				return;
+			}
+			
+			RequestAuthorityChange(*ReplicationManager, MoveTemp(AuthorityRequest), MoveTemp(NotificationToReuse))
+				.Next([Promise = MoveTemp(Promise), JoinResult = MoveTemp(JoinResult)](ConcertSyncClient::Replication::FAuthorityChangeResponse Response) mutable
+				{
+					Promise.EmplaceValue(FJoinSessionResult{ MoveTemp(JoinResult), MoveTemp(Response) });
+				});
+		});
+		return Future;
+	}
+
+	static void HandleJoinMessage(
+		TSharedPtr<SNotificationItem> NotificationItem,
+		ConcertSyncClient::Replication::FJoinReplicatedSessionResult ErrorResult
+		)
+	{
+		// We may already be on the game thread. In some cases, like a timeout, we'll be executed on the UDP's thread.
+		// NotificationItem may invalidate itself, which accesses the FSlateApplication singleton - which must be done on the game thread.
+		AsyncTask(ENamedThreads::GameThread, [NotificationItem = MoveTemp(NotificationItem), ErrorResult = MoveTemp(ErrorResult)]() mutable
+		{
+			const bool bSuccess = ErrorResult.ErrorCode == EJoinReplicationErrorCode::Success;
+			if (bSuccess)
+			{
+				NotificationItem->SetText(LOCTEXT("Joined", "Joined Replication Session"));
+				NotificationItem->SetCompletionState(SNotificationItem::CS_Success);
+			}
+			else
+			{
+				const FString ErrorCodeText = ConcertSyncCore::Replication::LexJoinErrorCode(ErrorResult.ErrorCode);
+				const FText Message = ErrorResult.DetailedErrorMessage.IsEmpty()
+					? FText::Format(LOCTEXT("Error.MissingInfoFmt", "Error code: {0}"), FText::FromString(ErrorCodeText))
+					: FText::Format(LOCTEXT("Error.DetailedInfoFmt", "{0} ({1})"), FText::FromString(ErrorResult.DetailedErrorMessage), FText::FromString(ErrorCodeText));
+				NotificationItem->SetSubText(Message);
+				NotificationItem->SetCompletionState(SNotificationItem::CS_Fail);
+			}
+		});
+	}
+	
+	static void HandleAuthorityMessage(
+		TSharedPtr<SNotificationItem> NotificationItem,
+		const int32 NumRequested,
+		const int32 NumReleased,
+		ConcertSyncClient::Replication::FAuthorityChangeResponse Response
+		)
+	{
+		// We may already be on the game thread. In some cases, like a timeout, we'll be executed on the UDP's thread.
+		// NotificationItem may invalidate itself, which accesses the FSlateApplication singleton - which must be done on the game thread.
+		AsyncTask(ENamedThreads::GameThread, [NotificationItem = MoveTemp(NotificationItem), NumRequested, NumReleased, Response]()
+		{
+			const int32 NumRejected = Response.RejectedObjects.Num();
+			const int32 NumTaken = NumRequested - NumRejected;
+					
+			FString FinalMessage;
+			if (NumTaken > 0)
+			{
+				FinalMessage.Append(FText::Format(LOCTEXT("AuthorityChange.Taken", "{0} owned objects."), NumTaken).ToString());
+			}
+			if (NumRejected > 0)
+			{
+				FinalMessage.Append(FText::Format(LOCTEXT("AuthorityChange.Rejected", "{0} rejected objects (see log)."), NumRejected).ToString());
+				for (const TPair<FSoftObjectPath, FConcertStreamArray>& RejectedObjectInfo : Response.RejectedObjects)
+				{
+					UE_LOG(LogConcert, Warning, TEXT("Rejected authority change for %s for %d streams"), *RejectedObjectInfo.Key.ToString(), RejectedObjectInfo.Value.StreamIds.Num());
+				}
+			}
+			if (NumReleased > 0)
+			{
+				FinalMessage.Append(FText::Format(LOCTEXT("AuthorityChange.Released", "{0} released objects."), NumReleased).ToString());
+			}
+
+			const bool bWasSuccess = Response.RejectedObjects.IsEmpty();
+			NotificationItem->SetText(bWasSuccess ? LOCTEXT("AuthorityChange.RequestAccepted", "Authority change successful.") : LOCTEXT("AuthorityChange.RequestErrorful", "Authority change has rejections"));
+			NotificationItem->SetSubText(FText::FromString(FinalMessage));
+			NotificationItem->SetCompletionState(bWasSuccess ? SNotificationItem::CS_Success : SNotificationItem::CS_Fail);
+			NotificationItem->ExpireAndFadeout();
+		});
+	}
+}
+
+#undef LOCTEXT_NAMESPACE
\ No newline at end of file
diff --git a/Engine/Plugins/Developer/Concert/ConcertApp/MultiUserClient/Source/MultiUserClient/Private/Replication/ReplicationUtils.h b/Engine/Plugins/Developer/Concert/ConcertApp/MultiUserClient/Source/MultiUserClient/Private/Replication/ReplicationUtils.h
new file mode 100644
--- /dev/null
+++ b/Engine/Plugins/Developer/Concert/ConcertApp/MultiUserClient/Source/MultiUserClient/Private/Replication/ReplicationUtils.h
@@ -0,0 +1,58 @@
+// Copyright Epic Games, Inc. All Rights Reserved.
+
+#pragma once
+
+#include "Replication/IConcertClientReplicationManager.h"
+
+#include "Async/Future.h"
+#include "Misc/Optional.h"
+#include "Templates/SharedPointer.h"
+
+class IConcertSyncClient;
+class IConcertClientReplicationManager;
+class SNotificationItem;
+
+namespace UE::MultiUserClient::Replication
+{
+	struct FJoinSessionResult
+	{
+		/** Result of joining the replication session. */
+		ConcertSyncClient::Replication::FJoinReplicatedSessionResult JoinRequestResult;
+
+		/** Set if an authority request was sent. */
+		TOptional<ConcertSyncClient::Replication::FAuthorityChangeResponse> AuthorityResponse;
+	};
+	
+	/**
+	 * Decorates a replication join request with additional logic specific to Multi-User workflows.
+	 *
+	 * Extra steps involve
+	 *	- Adds an editor UI notifications on the bottom-right: One for the duration of the process, and then another for
+	 *	either success (leaves after a few seconds) or failure (must be dismissed).
+	 *	- Automatically takes authority over the objects in the stream (TODO DP ... if it is configured in the multi user settings).
+	 *
+	 * @note The future can finish on any thread so take care to synchronize if needed!	
+	 * @return Future that executes when all actions have been performed (joining, optional authority requests, etc.)
+	 */
+	TFuture<FJoinSessionResult> JoinSessionForMultiUser(
+		const TSharedRef<IConcertSyncClient>& Client,
+		ConcertSyncClient::Replication::FJoinReplicatedSessionArgs Args
+		);
+
+	/**
+	 * Decorates a request for changing authority with additional logic specific to Multi-User workflows.
+	 * @note The future can finish on any thread so take care to synchronize if needed!
+	 * @return Future that executes when the authority response has been received from the server.
+	 */
+	TFuture<ConcertSyncClient::Replication::FAuthorityChangeResponse> RequestAuthorityChange(
+		IConcertClientReplicationManager& Manager,
+		ConcertSyncClient::Replication::FAuthorityChangeRequest Args,
+		TSharedPtr<SNotificationItem> NotificationToReuse = nullptr
+		);
+
+	/**
+	 * Decorates a replication leave requests with additional logic specific to Multi-User workflows.
+	 * This shows an editor UI notification on the bottom-right.
+	 */
+	void LeaveSessionForMultiUser(IConcertClientReplicationManager& Manager);
+}
diff --git a/Engine/Plugins/Developer/Concert/ConcertApp/MultiUserClient/Source/MultiUserClient/Private/Util/ReplicationUtils.cpp b/Engine/Plugins/Developer/Concert/ConcertApp/MultiUserClient/Source/MultiUserClient/Private/Util/ReplicationUtils.cpp
deleted file mode 100644
diff --git a/Engine/Plugins/Developer/Concert/ConcertApp/MultiUserClient/Source/MultiUserClient/Private/Util/ReplicationUtils.h b/Engine/Plugins/Developer/Concert/ConcertApp/MultiUserClient/Source/MultiUserClient/Private/Util/ReplicationUtils.h
deleted file mode 100644
diff --git a/Engine/Plugins/Developer/Concert/ConcertApp/MultiUserClient/Source/MultiUserClient/Private/Widgets/ActiveSession/Replication/SReplicationConnectionControls.cpp b/Engine/Plugins/Developer/Concert/ConcertApp/MultiUserClient/Source/MultiUserClient/Private/Widgets/ActiveSession/Replication/SReplicationConnectionControls.cpp
--- a/Engine/Plugins/Developer/Concert/ConcertApp/MultiUserClient/Source/MultiUserClient/Private/Widgets/ActiveSession/Replication/SReplicationConnectionControls.cpp
+++ b/Engine/Plugins/Developer/Concert/ConcertApp/MultiUserClient/Source/MultiUserClient/Private/Widgets/ActiveSession/Replication/SReplicationConnectionControls.cpp
@@ -2,12 +2,9 @@
 
 #include "SReplicationConnectionControls.h"
 
-#include "IConcertSyncClient.h"
-#include "IConcertSyncClientModule.h"
 #include "Customization/ShowJoinSettingsCustomization.h"
-#include "Replication/IConcertClientReplicationManager.h"
+#include "Replication/MultiUserReplicationManager.h"
 #include "Settings/MultiUserReplicationSettings.h"
-#include "Util/ReplicationUtils.h"
 
 #include "IDetailsView.h"
 #include "Modules/ModuleManager.h"
@@ -29,8 +26,10 @@
 
 namespace UE::MultiUserClient
 {
-	void SReplicationConnectionControls::Construct(const FArguments& InArgs)
+	void SReplicationConnectionControls::Construct(const FArguments& InArgs, TSharedRef<FMultiUserReplicationManager> InReplicationManager)
 	{
+		ReplicationManager = InReplicationManager;
+		
 		TSharedPtr<SVerticalBox> Content;
 		ChildSlot
 		[
@@ -123,17 +122,8 @@ namespace UE::MultiUserClient
 
 	FReply SReplicationConnectionControls::OnJoinButtonClicked()
 	{
-		TSharedPtr<IConcertSyncClient> Client = IConcertSyncClientModule::Get().GetClient(TEXT("MultiUser"));
-		IConcertClientReplicationManager* ReplicationManager = Client ? Client->GetReplicationManager() : nullptr;
-		if (!ensure(ReplicationManager))
-		{
-			return FReply::Handled();
-		}
-
-		UMultiUserReplicationClientProfileAsset* ClientProfile = GetClientProfile();
-		checkf(ClientProfile, TEXT("Validate IsJoinButtonEnabled implementation!"));
 		// All the UI callbacks for visibility, etc. will update automatically since they do polling.
-		ReplicationUtils::JoinSessionWithEditorNotifications(*ReplicationManager, ClientProfile->ToJoinArgs());
+		ReplicationManager->JoinReplicationSession(*GetClientProfile());
 		return FReply::Handled();
 	}
 
@@ -153,9 +143,7 @@ namespace UE::MultiUserClient
 			return false;
 		}
 		
-		TSharedPtr<IConcertSyncClient> Client = IConcertSyncClientModule::Get().GetClient(TEXT("MultiUser"));
-		IConcertClientReplicationManager* ReplicationManager = Client ? Client->GetReplicationManager() : nullptr;
-		const bool bCanMakeJoinRequest = ReplicationManager && ReplicationManager->CanJoin();
+		const bool bCanMakeJoinRequest = ReplicationManager->CanJoin();
 		if (!bCanMakeJoinRequest && ErrorReason)
 		{
 			*ErrorReason = LOCTEXT("IsJoinButtonEnabled.Error.JoinRequestInProgress", "A join request is in progress.");
@@ -167,20 +155,13 @@ namespace UE::MultiUserClient
 
 	EVisibility SReplicationConnectionControls::GetJoinButtonVisibility() const
 	{
-		TSharedPtr<IConcertSyncClient> Client = IConcertSyncClientModule::Get().GetClient(TEXT("MultiUser"));
-		IConcertClientReplicationManager* ReplicationManager = Client ? Client->GetReplicationManager() : nullptr;
-		const bool bShouldBeVisible = !ReplicationManager || !ReplicationManager->IsConnectedToReplicationSession();
+		const bool bShouldBeVisible = !ReplicationManager->IsConnectedToReplicationSession();
 		return bShouldBeVisible ? EVisibility::Visible : EVisibility::Collapsed;
 	}
 
 	FReply SReplicationConnectionControls::OnLeaveButtonClicked()
 	{
-		TSharedPtr<IConcertSyncClient> Client = IConcertSyncClientModule::Get().GetClient(TEXT("MultiUser"));
-		IConcertClientReplicationManager* ReplicationManager = Client ? Client->GetReplicationManager() : nullptr;
-		if (ensure(ReplicationManager))
-		{
-			ReplicationUtils::LeaveSessionWithEditorNotifications(*ReplicationManager);
-		}
+		ReplicationManager->LeaveSession();
 		return FReply::Handled();
 	}
 }
diff --git a/Engine/Plugins/Developer/Concert/ConcertApp/MultiUserClient/Source/MultiUserClient/Private/Widgets/ActiveSession/Replication/SReplicationConnectionControls.h b/Engine/Plugins/Developer/Concert/ConcertApp/MultiUserClient/Source/MultiUserClient/Private/Widgets/ActiveSession/Replication/SReplicationConnectionControls.h
--- a/Engine/Plugins/Developer/Concert/ConcertApp/MultiUserClient/Source/MultiUserClient/Private/Widgets/ActiveSession/Replication/SReplicationConnectionControls.h
+++ b/Engine/Plugins/Developer/Concert/ConcertApp/MultiUserClient/Source/MultiUserClient/Private/Widgets/ActiveSession/Replication/SReplicationConnectionControls.h
@@ -11,6 +11,8 @@ class UJoinReplicationSessionArgs;
 
 namespace UE::MultiUserClient
 {
+	class FMultiUserReplicationManager;
+
 	class SReplicationConnectionControls : public SCompoundWidget
 	{
 	public:
@@ -19,10 +21,13 @@ namespace UE::MultiUserClient
 		{}
 		SLATE_END_ARGS()
 
-		void Construct(const FArguments& InArgs);
+		void Construct(const FArguments& InArgs, TSharedRef<FMultiUserReplicationManager> InReplicationManager);
 
 	private:
 
+		/** Handles replication requests for Multi-User. */
+		TSharedPtr<FMultiUserReplicationManager> ReplicationManager;
+
 		UMultiUserReplicationClientProfileAsset* GetClientProfile() const;
 
 		// Join Button
diff --git a/Engine/Plugins/Developer/Concert/ConcertApp/MultiUserClient/Source/MultiUserClient/Private/Widgets/ActiveSession/Replication/SReplicationControlsTab.cpp b/Engine/Plugins/Developer/Concert/ConcertApp/MultiUserClient/Source/MultiUserClient/Private/Widgets/ActiveSession/Replication/SReplicationControlsTab.cpp
--- a/Engine/Plugins/Developer/Concert/ConcertApp/MultiUserClient/Source/MultiUserClient/Private/Widgets/ActiveSession/Replication/SReplicationControlsTab.cpp
+++ b/Engine/Plugins/Developer/Concert/ConcertApp/MultiUserClient/Source/MultiUserClient/Private/Widgets/ActiveSession/Replication/SReplicationControlsTab.cpp
@@ -54,7 +54,7 @@ namespace UE::MultiUserClient
 		}
 	}
 	
-	void SReplicationControlsTab::Construct(const FArguments& InArgs)
+	void SReplicationControlsTab::Construct(const FArguments& InArgs, TSharedRef<FMultiUserReplicationManager> InReplicationManager)
 	{
 		constexpr bool bCollpaseInitially_ConnectionArea	= false;
 		constexpr bool bCollpaseInitially_AttributesArea	= true;
@@ -75,7 +75,7 @@ namespace UE::MultiUserClient
 					FOnBooleanValueChanged::CreateSP(this, &SReplicationControlsTab::OnConnectionAreaExpansionChanged),
 					LOCTEXT("ConnectionArea.Label", "Connection"),
 					LOCTEXT("ConnectionArea.Tooltip", "Control the replication connection, e.g. join and leave"),
-					SNew(SReplicationConnectionControls),
+					SNew(SReplicationConnectionControls, InReplicationManager),
 					bCollpaseInitially_ConnectionArea
 					)
 			]
diff --git a/Engine/Plugins/Developer/Concert/ConcertApp/MultiUserClient/Source/MultiUserClient/Private/Widgets/ActiveSession/Replication/SReplicationControlsTab.h b/Engine/Plugins/Developer/Concert/ConcertApp/MultiUserClient/Source/MultiUserClient/Private/Widgets/ActiveSession/Replication/SReplicationControlsTab.h
--- a/Engine/Plugins/Developer/Concert/ConcertApp/MultiUserClient/Source/MultiUserClient/Private/Widgets/ActiveSession/Replication/SReplicationControlsTab.h
+++ b/Engine/Plugins/Developer/Concert/ConcertApp/MultiUserClient/Source/MultiUserClient/Private/Widgets/ActiveSession/Replication/SReplicationControlsTab.h
@@ -10,6 +10,8 @@ class SExpandableArea;
 
 namespace UE::MultiUserClient
 {
+	class FMultiUserReplicationManager;
+
 	class SReplicationControlsTab : public SCompoundWidget
 	{
 	public:
@@ -18,7 +20,7 @@ namespace UE::MultiUserClient
 		{}
 		SLATE_END_ARGS()
 
-		void Construct(const FArguments& InArgs);
+		void Construct(const FArguments& InArgs, TSharedRef<FMultiUserReplicationManager> InReplicationManager);
 
 	private:
 
diff --git a/Engine/Plugins/Developer/Concert/ConcertApp/MultiUserClient/Source/MultiUserClient/Private/Widgets/ActiveSession/Replication/SReplicationTabWithWarningOverlay.cpp b/Engine/Plugins/Developer/Concert/ConcertApp/MultiUserClient/Source/MultiUserClient/Private/Widgets/ActiveSession/Replication/SReplicationTabWithWarningOverlay.cpp
--- a/Engine/Plugins/Developer/Concert/ConcertApp/MultiUserClient/Source/MultiUserClient/Private/Widgets/ActiveSession/Replication/SReplicationTabWithWarningOverlay.cpp
+++ b/Engine/Plugins/Developer/Concert/ConcertApp/MultiUserClient/Source/MultiUserClient/Private/Widgets/ActiveSession/Replication/SReplicationTabWithWarningOverlay.cpp
@@ -14,7 +14,7 @@
 
 namespace UE::MultiUserClient
 {
-	void SReplicationTabWithWarningOverlay::Construct(const FArguments& InArgs)
+	void SReplicationTabWithWarningOverlay::Construct(const FArguments& InArgs, TSharedRef<FMultiUserReplicationManager> InReplicationManager)
 	{
 		ChildSlot
 		[
@@ -23,7 +23,7 @@ namespace UE::MultiUserClient
 			// The real content
 			+SOverlay::Slot()
 			[
-				SNew(SReplicationControlsTab)
+				SNew(SReplicationControlsTab, InReplicationManager)
 			]
 
 			//  
diff --git a/Engine/Plugins/Developer/Concert/ConcertApp/MultiUserClient/Source/MultiUserClient/Private/Widgets/ActiveSession/Replication/SReplicationTabWithWarningOverlay.h b/Engine/Plugins/Developer/Concert/ConcertApp/MultiUserClient/Source/MultiUserClient/Private/Widgets/ActiveSession/Replication/SReplicationTabWithWarningOverlay.h
--- a/Engine/Plugins/Developer/Concert/ConcertApp/MultiUserClient/Source/MultiUserClient/Private/Widgets/ActiveSession/Replication/SReplicationTabWithWarningOverlay.h
+++ b/Engine/Plugins/Developer/Concert/ConcertApp/MultiUserClient/Source/MultiUserClient/Private/Widgets/ActiveSession/Replication/SReplicationTabWithWarningOverlay.h
@@ -9,6 +9,8 @@ class SOverlay;
 
 namespace UE::MultiUserClient
 {
+	class FMultiUserReplicationManager;
+
 	/** Wraps SReplicationControlsTab with an overlay warning the user that this feature is still experimental */
 	class SReplicationTabWithWarningOverlay : public SCompoundWidget
 	{
@@ -18,7 +20,7 @@ namespace UE::MultiUserClient
 		{}
 		SLATE_END_ARGS()
 
-		void Construct(const FArguments& InArgs);
+		void Construct(const FArguments& InArgs, TSharedRef<FMultiUserReplicationManager> InReplicationManager);
 		
 	private:
 
diff --git a/Engine/Plugins/Developer/Concert/ConcertApp/MultiUserClient/Source/MultiUserClient/Private/Widgets/ActiveSession/SActiveSessionRoot.cpp b/Engine/Plugins/Developer/Concert/ConcertApp/MultiUserClient/Source/MultiUserClient/Private/Widgets/ActiveSession/SActiveSessionRoot.cpp
--- a/Engine/Plugins/Developer/Concert/ConcertApp/MultiUserClient/Source/MultiUserClient/Private/Widgets/ActiveSession/SActiveSessionRoot.cpp
+++ b/Engine/Plugins/Developer/Concert/ConcertApp/MultiUserClient/Source/MultiUserClient/Private/Widgets/ActiveSession/SActiveSessionRoot.cpp
@@ -23,7 +23,8 @@ namespace UE::MultiUserClient
 	void SActiveSessionRoot::Construct(
 		const FArguments& InArgs,
 		const TSharedRef<SDockTab>& ConstructUnderMajorTab,
-		TSharedPtr<IConcertSyncClient> InConcertSyncClient
+		TSharedPtr<IConcertSyncClient> InConcertSyncClient,
+		TSharedRef<FMultiUserReplicationManager> InReplicationManager
 		)
 	{
 		ConcertSyncClient = InConcertSyncClient;
@@ -31,7 +32,7 @@ namespace UE::MultiUserClient
 		TabManager = FGlobalTabmanager::Get()->NewTabManager(ConstructUnderMajorTab);
 		TSharedRef<FWorkspaceItem> AppMenuGroup = TabManager->AddLocalWorkspaceMenuCategory(LOCTEXT("ConcertActiveSession", "Active Session"));
 		TabManager->SetAllowWindowMenuBar(true);
-		RegisterTabSpawners(TabManager.ToSharedRef(), AppMenuGroup);
+		RegisterTabSpawners(TabManager.ToSharedRef(), AppMenuGroup, InReplicationManager);
 
 		// Create our content
 		const TSharedRef<FTabManager::FLayout> Layout =
@@ -76,13 +77,17 @@ namespace UE::MultiUserClient
 		];
 	}
 
-	void SActiveSessionRoot::RegisterTabSpawners(const TSharedRef<FTabManager>& InTabManager, const TSharedRef<FWorkspaceItem>& AppMenuGroup)
+	void SActiveSessionRoot::RegisterTabSpawners(
+		const TSharedRef<FTabManager>& InTabManager,
+		const TSharedRef<FWorkspaceItem>& AppMenuGroup,
+		TSharedRef<FMultiUserReplicationManager> InReplicationManager
+		)
 	{
 		InTabManager->RegisterTabSpawner(OverviewTabId, FOnSpawnTab::CreateSP(this, &SActiveSessionRoot::SpawnTab_Overview))
 			.SetDisplayName(LOCTEXT("OverviewTab.DisplayName", "Overview"))
 			.SetGroup(AppMenuGroup);
 	
-		InTabManager->RegisterTabSpawner(ReplicationTabId, FOnSpawnTab::CreateSP(this, &SActiveSessionRoot::SpawnTab_ReplicationControls))
+		InTabManager->RegisterTabSpawner(ReplicationTabId, FOnSpawnTab::CreateSP(this, &SActiveSessionRoot::SpawnTab_ReplicationControls, InReplicationManager))
 			.SetDisplayName(LOCTEXT("ReplicationTab.DisplayName", "Replication"))
 			.SetGroup(AppMenuGroup);
 	}
@@ -100,13 +105,13 @@ namespace UE::MultiUserClient
 			];
 	}
 
-	TSharedRef<SDockTab> SActiveSessionRoot::SpawnTab_ReplicationControls(const FSpawnTabArgs& Args)
+	TSharedRef<SDockTab> SActiveSessionRoot::SpawnTab_ReplicationControls(const FSpawnTabArgs& Args, TSharedRef<FMultiUserReplicationManager> InReplicationManager)
 	{
 		return SNew(SDockTab)
 			.Label(LOCTEXT("ReplicationTab.Label", "Replication"))
 			.ToolTipText(LOCTEXT("ReplicationTab.Tooltip", "Manage real-time object replication"))
 			[
-				SNew(SReplicationTabWithWarningOverlay)
+				SNew(SReplicationTabWithWarningOverlay, InReplicationManager)
 			];
 	}
 }
diff --git a/Engine/Plugins/Developer/Concert/ConcertApp/MultiUserClient/Source/MultiUserClient/Private/Widgets/ActiveSession/SActiveSessionRoot.h b/Engine/Plugins/Developer/Concert/ConcertApp/MultiUserClient/Source/MultiUserClient/Private/Widgets/ActiveSession/SActiveSessionRoot.h
--- a/Engine/Plugins/Developer/Concert/ConcertApp/MultiUserClient/Source/MultiUserClient/Private/Widgets/ActiveSession/SActiveSessionRoot.h
+++ b/Engine/Plugins/Developer/Concert/ConcertApp/MultiUserClient/Source/MultiUserClient/Private/Widgets/ActiveSession/SActiveSessionRoot.h
@@ -14,6 +14,8 @@ class SWindow;
 
 namespace UE::MultiUserClient
 {
+	class FMultiUserReplicationManager;
+	
 	/**
 	 * Displayed when the client is connected to an active session.
 	 * Manages the child content in tabs.
@@ -29,7 +31,7 @@ namespace UE::MultiUserClient
 		{}
 		SLATE_END_ARGS()
 
-		void Construct(const FArguments& InArgs, const TSharedRef<SDockTab>& ConstructUnderMajorTab, TSharedPtr<IConcertSyncClient> InConcertSyncClient);
+		void Construct(const FArguments& InArgs, const TSharedRef<SDockTab>& ConstructUnderMajorTab, TSharedPtr<IConcertSyncClient> InConcertSyncClient, TSharedRef<FMultiUserReplicationManager> InReplicationManager);
 
 	private:
 
@@ -37,9 +39,9 @@ namespace UE::MultiUserClient
 		/** Holds the child content: SActiveSessionOverviewTab and SReplicationControlsTab. */
 		TSharedPtr<FTabManager> TabManager;
 		
-		void RegisterTabSpawners(const TSharedRef<FTabManager>& InTabManager, const TSharedRef<FWorkspaceItem>& AppMenuGroup);
+		void RegisterTabSpawners(const TSharedRef<FTabManager>& InTabManager, const TSharedRef<FWorkspaceItem>& AppMenuGroup, TSharedRef<FMultiUserReplicationManager> InReplicationManager);
 		TSharedRef<SDockTab> SpawnTab_Overview(const FSpawnTabArgs& Args);
-		TSharedRef<SDockTab> SpawnTab_ReplicationControls(const FSpawnTabArgs& Args);
+		TSharedRef<SDockTab> SpawnTab_ReplicationControls(const FSpawnTabArgs& Args, TSharedRef<FMultiUserReplicationManager> InReplicationManager);
 	};
 }
 
diff --git a/Engine/Plugins/Developer/Concert/ConcertApp/MultiUserClient/Source/MultiUserClient/Private/Widgets/SConcertBrowser.cpp b/Engine/Plugins/Developer/Concert/ConcertApp/MultiUserClient/Source/MultiUserClient/Private/Widgets/SConcertBrowser.cpp
--- a/Engine/Plugins/Developer/Concert/ConcertApp/MultiUserClient/Source/MultiUserClient/Private/Widgets/SConcertBrowser.cpp
+++ b/Engine/Plugins/Developer/Concert/ConcertApp/MultiUserClient/Source/MultiUserClient/Private/Widgets/SConcertBrowser.cpp
@@ -18,7 +18,11 @@
 
 #define LOCTEXT_NAMESPACE "SConcertBrowser"
 
-void SConcertBrowser::Construct(const FArguments& InArgs, TSharedRef<SDockTab> InConstructUnderMajorTab, TWeakPtr<IConcertSyncClient> InSyncClient)
+void SConcertBrowser::Construct(
+	const FArguments& InArgs,
+	TSharedRef<SDockTab> InConstructUnderMajorTab,
+	TSharedRef<IConcertSyncClient> InSyncClient,
+	TSharedRef<UE::MultiUserClient::FMultiUserReplicationManager> InReplicationManager)
 {
 	if (!MultiUserClientUtils::HasServerCompatibleCommunicationPluginEnabled())
 	{
@@ -33,6 +37,7 @@ void SConcertBrowser::Construct(const FArguments& InArgs, TSharedRef<SDockTab> I
 	}
 
 	WeakConcertSyncClient = InSyncClient;
+	WeakReplicationManager = InReplicationManager;
 	ConstructedUnderMajorTab = MoveTemp(InConstructUnderMajorTab);
 	if (TSharedPtr<IConcertSyncClient> ConcertSyncClient = WeakConcertSyncClient.Pin())
 	{
@@ -55,15 +60,31 @@ void SConcertBrowser::HandleSessionConnectionChanged(IConcertClientSession& InSe
 
 void SConcertBrowser::AttachChildWidget(EConcertConnectionStatus ConnectionStatus)
 {
-	if (TSharedPtr<IConcertSyncClient> ConcertSyncClient = WeakConcertSyncClient.Pin())
+	if (const TSharedPtr<IConcertSyncClient> ConcertSyncClient = WeakConcertSyncClient.Pin()
+		; ensure(ConcertSyncClient))
 	{
 		if (ConnectionStatus == EConcertConnectionStatus::Connected)
 		{
-			ChildSlot.AttachWidget(SNew(UE::MultiUserClient::SActiveSessionRoot, ConstructedUnderMajorTab.ToSharedRef(), ConcertSyncClient));
+			if (const TSharedPtr<UE::MultiUserClient::FMultiUserReplicationManager> ReplicationManager = WeakReplicationManager.Pin()
+				; ensure(ReplicationManager))
+			{
+				ChildSlot.AttachWidget(
+					SNew(UE::MultiUserClient::SActiveSessionRoot,
+						ConstructedUnderMajorTab.ToSharedRef(),
+						ConcertSyncClient,
+						ReplicationManager.ToSharedRef()
+						)
+					);
+			}
 		}
 		else if (ConnectionStatus == EConcertConnectionStatus::Disconnected)
 		{
-			ChildSlot.AttachWidget(SNew(SConcertClientSessionBrowser, ConcertSyncClient->GetConcertClient(), SearchedText));
+			ChildSlot.AttachWidget(
+				SNew(SConcertClientSessionBrowser,
+					ConcertSyncClient->GetConcertClient(),
+					SearchedText
+					)
+				);
 		}
 	}
 }
diff --git a/Engine/Plugins/Developer/Concert/ConcertApp/MultiUserClient/Source/MultiUserClient/Private/Widgets/SConcertBrowser.h b/Engine/Plugins/Developer/Concert/ConcertApp/MultiUserClient/Source/MultiUserClient/Private/Widgets/SConcertBrowser.h
--- a/Engine/Plugins/Developer/Concert/ConcertApp/MultiUserClient/Source/MultiUserClient/Private/Widgets/SConcertBrowser.h
+++ b/Engine/Plugins/Developer/Concert/ConcertApp/MultiUserClient/Source/MultiUserClient/Private/Widgets/SConcertBrowser.h
@@ -9,6 +9,11 @@
 #include "IConcertModule.h"  // Change to use Fwd or Ptr.h?
 #include "ConcertMessages.h"
 
+namespace UE::MultiUserClient
+{
+	class FMultiUserReplicationManager;
+}
+
 class IConcertClientSession;
 class IConcertSyncClient;
 
@@ -32,12 +37,14 @@ public:
 	* @param InConstructUnderWindow The window in which this widget is being constructed.
 	* @param InSyncClient The sync client.
 	*/
-	void Construct(const FArguments& InArgs, TSharedRef<SDockTab> InConstructUnderMajorTab, TWeakPtr<IConcertSyncClient> InSyncClient);
+	void Construct(const FArguments& InArgs, TSharedRef<SDockTab> InConstructUnderMajorTab, TSharedRef<IConcertSyncClient> InSyncClient, TSharedRef<UE::MultiUserClient::FMultiUserReplicationManager> InReplicationManager);
 
 private:
 
 	/** Keeps the sync client interface. */
 	TWeakPtr<IConcertSyncClient> WeakConcertSyncClient;
+	/** Interacts with the replication system on behalf of Multi-User. */
+	TWeakPtr<UE::MultiUserClient::FMultiUserReplicationManager> WeakReplicationManager;
 
 	/** Kept so it can be passed on to SActiveSessionRoot. */
 	TSharedPtr<SDockTab> ConstructedUnderMajorTab;
diff --git a/Engine/Plugins/Developer/Concert/ConcertApp/MultiUserClient/Source/MultiUserReplication/Public/MultiUserReplicationClientProfileAsset.h b/Engine/Plugins/Developer/Concert/ConcertApp/MultiUserClient/Source/MultiUserReplication/Public/MultiUserReplicationClientProfileAsset.h
--- a/Engine/Plugins/Developer/Concert/ConcertApp/MultiUserClient/Source/MultiUserReplication/Public/MultiUserReplicationClientProfileAsset.h
+++ b/Engine/Plugins/Developer/Concert/ConcertApp/MultiUserClient/Source/MultiUserReplication/Public/MultiUserReplicationClientProfileAsset.h
@@ -27,7 +27,7 @@ public:
 
 	/** The streams this client will send to the server. It takes authority over the listed properties. */
 	UPROPERTY(EditAnywhere, Category = "Replication")
-	TArray<TObjectPtr<UMultiUserReplicationStreamAsset>> Streams;
+	TSet<TObjectPtr<UMultiUserReplicationStreamAsset>> Streams;
 
 	/** Converts this asset's data to the args expected by IConcertClientReplicationManager for joining. */
 	UE::ConcertSyncClient::Replication::FJoinReplicatedSessionArgs ToJoinArgs() const;
diff --git a/Engine/Plugins/Developer/Concert/ConcertSync/ConcertSyncClient/Source/ConcertSyncClient/Private/Replication/IConcertClientReplicationManager.cpp b/Engine/Plugins/Developer/Concert/ConcertSync/ConcertSyncClient/Source/ConcertSyncClient/Private/Replication/IConcertClientReplicationManager.cpp
new file mode 100644
--- /dev/null
+++ b/Engine/Plugins/Developer/Concert/ConcertSync/ConcertSyncClient/Source/ConcertSyncClient/Private/Replication/IConcertClientReplicationManager.cpp
@@ -0,0 +1,98 @@
+// Copyright Epic Games, Inc. All Rights Reserved.
+
+#include "Replication/IConcertClientReplicationManager.h"
+
+#include "ConcertLogGlobal.h"
+
+#include "Async/Future.h"
+
+namespace UE::ConcertSyncClient::Replication::Private
+{
+	template<typename Callback>
+	void ForEachStreamContainingObject(TArrayView<const FSoftObjectPath> Objects, const IConcertClientReplicationManager& Manager, Callback InCallback)
+	{
+		for (const FSoftObjectPath& ObjectPath : Objects)
+		{
+			Manager.ForEachRegisteredStream([&InCallback, &ObjectPath](const FReplicationStreamDescription& StreamDescription)
+			{
+				const bool bStreamContainsObject = StreamDescription.BaseDescription.ReplicationMap.ReplicatedObjects.Contains(ObjectPath);
+				if (bStreamContainsObject)
+				{
+					InCallback(ObjectPath, StreamDescription.BaseDescription.Identifier);
+				}
+				else
+				{
+					UE_LOG(LogConcert, Warning, TEXT("Object %s is not a valid argument because it is not contained in any stream."), *ObjectPath.ToString());
+				}
+				return EBreakBehavior::Continue;
+			});
+		}
+	}
+}
+
+bool IConcertClientReplicationManager::HasRegisteredStreams() const
+{
+	return ForEachRegisteredStream([](const auto&){ return EBreakBehavior::Break; }) == EStreamEnumerationResult::Iterated;
+}
+
+TFuture<UE::ConcertSyncClient::Replication::FAuthorityChangeResponse> IConcertClientReplicationManager::TakeAuthorityOver(TArrayView<const FSoftObjectPath> Objects)
+{
+	using namespace UE::ConcertSyncClient::Replication;
+	
+	if (!HasRegisteredStreams())
+	{
+		UE_LOG(LogConcert, Error, TEXT("Attempted to take authority while not connected!"));
+		TMap<FSoftObjectPath, FConcertStreamArray> Result;
+		Algo::Transform(Objects, Result, [](const FSoftObjectPath& Path){ return Path; });
+		return MakeFulfilledPromise<FAuthorityChangeResponse>(FAuthorityChangeResponse{{ MoveTemp(Result) }}).GetFuture();
+	}
+
+	FAuthorityChangeRequest Request;
+	Private::ForEachStreamContainingObject(Objects, *this,
+		[&Request](const FSoftObjectPath& ObjectPath, const FGuid& StreamId)
+		{
+			Request.TakeAuthority.FindOrAdd(ObjectPath).StreamIds.Add(StreamId);
+		});
+
+	// Do not send pointless, empty requests to the server
+	if (Request.TakeAuthority.IsEmpty())
+	{
+		// Not only does this warn about incorrect API use at runtime, this also helps debug (incorrectly written) unit tests
+		const FString ObjectsAsString = FString::JoinBy(Objects, TEXT(","), [](const FSoftObjectPath& Path){ return Path.ToString(); });
+		UE_LOG(LogConcert, Warning, TEXT("Local client did not registered any stream for the given objects. This take authority request will not be sent. Objects: %s"), *ObjectsAsString);
+		return MakeFulfilledPromise<FAuthorityChangeResponse>(FAuthorityChangeResponse{}).GetFuture();
+	}
+	
+	return RequestAuthorityChange(MoveTemp(Request));
+}
+
+TFuture<UE::ConcertSyncClient::Replication::FAuthorityChangeResponse> IConcertClientReplicationManager::ReleaseAuthorityOf(TArrayView<const FSoftObjectPath> Objects)
+{
+	using namespace UE::ConcertSyncClient::Replication;
+	
+	if (!HasRegisteredStreams())
+	{
+		UE_LOG(LogConcert, Error, TEXT("Attempted to take authority while not connected!"));
+		TMap<FSoftObjectPath, FConcertStreamArray> Result;
+		Algo::Transform(Objects, Result, [](const FSoftObjectPath& Path){ return Path; });
+		return MakeFulfilledPromise<FAuthorityChangeResponse>(FAuthorityChangeResponse{{ MoveTemp(Result) }}).GetFuture();
+	}
+	
+	FAuthorityChangeRequest Request;
+	Private::ForEachStreamContainingObject(Objects, *this,
+		[&Request](const FSoftObjectPath& ObjectPath, const FGuid& StreamId)
+		{
+			Request.ReleaseAuthority.FindOrAdd(ObjectPath).StreamIds.Add(StreamId);
+		});
+
+	// Do not send pointless, empty requests to the server
+	if (Request.ReleaseAuthority.IsEmpty())
+	{
+		// Not only does this warn about incorrect API use at runtime, this also helps debug (incorrectly written) unit tests
+		const FString ObjectsAsString = FString::JoinBy(Objects, TEXT(","), [](const FSoftObjectPath& Path){ return Path.ToString(); });
+		UE_LOG(LogConcert, Warning, TEXT("Local client did not register any stream for the given objects. This release authority request will not be sent. Objects: %s"), *ObjectsAsString);
+		return MakeFulfilledPromise<FAuthorityChangeResponse>(FAuthorityChangeResponse{}).GetFuture();
+	}
+	
+	return RequestAuthorityChange(MoveTemp(Request));
+}
\ No newline at end of file
diff --git a/Engine/Plugins/Developer/Concert/ConcertSync/ConcertSyncClient/Source/ConcertSyncClient/Private/Replication/Manager/ReplicationManager.cpp b/Engine/Plugins/Developer/Concert/ConcertSync/ConcertSyncClient/Source/ConcertSyncClient/Private/Replication/Manager/ReplicationManager.cpp
--- a/Engine/Plugins/Developer/Concert/ConcertSync/ConcertSyncClient/Source/ConcertSyncClient/Private/Replication/Manager/ReplicationManager.cpp
+++ b/Engine/Plugins/Developer/Concert/ConcertSync/ConcertSyncClient/Source/ConcertSyncClient/Private/Replication/Manager/ReplicationManager.cpp
@@ -2,10 +2,9 @@
 
 #include "ReplicationManager.h"
 
-#include "ConcertLogGlobal.h"
-#include "ConcertSyncClientLiveSession.h"
 #include "IConcertSession.h"
 #include "ReplicationManagerState_Disconnected.h"
+#include "ReplicationManagerUtils.h"
 
 #include "Replication/Messages/ConcertReplicationHandshakeMessages.h"
 
@@ -52,6 +51,25 @@ namespace UE::ConcertSyncClient::Replication
 			&& CurrentState->IsConnectedToReplicationSession(); 
 	}
 
+	IConcertClientReplicationManager::EStreamEnumerationResult FReplicationManager::ForEachRegisteredStream(
+		TFunctionRef<EBreakBehavior(const FReplicationStreamDescription& Stream)> Callback
+		) const
+	{
+		return ensureMsgf(CurrentState, TEXT("StartAcceptingJoinRequests should have been called at this point."))
+			? CurrentState->ForEachRegisteredStream(Callback)
+			: EStreamEnumerationResult::NoRegisteredStreams;
+	}
+
+	TFuture<FAuthorityChangeResponse> FReplicationManager::RequestAuthorityChange(FAuthorityChangeRequest Args)
+	{
+		if (ensureMsgf(CurrentState, TEXT("StartAcceptingJoinRequests should have been called at this point.")))
+		{
+			return CurrentState->RequestAuthorityChange(Args);
+		}
+		
+		return RejectAll(MoveTemp(Args));
+	}
+
 	void FReplicationManager::OnChangeState(TSharedRef<FReplicationManagerState> NewState)
 	{
 		CurrentState = MoveTemp(NewState);
diff --git a/Engine/Plugins/Developer/Concert/ConcertSync/ConcertSyncClient/Source/ConcertSyncClient/Private/Replication/Manager/ReplicationManager.h b/Engine/Plugins/Developer/Concert/ConcertSync/ConcertSyncClient/Source/ConcertSyncClient/Private/Replication/Manager/ReplicationManager.h
--- a/Engine/Plugins/Developer/Concert/ConcertSync/ConcertSyncClient/Source/ConcertSyncClient/Private/Replication/Manager/ReplicationManager.h
+++ b/Engine/Plugins/Developer/Concert/ConcertSync/ConcertSyncClient/Source/ConcertSyncClient/Private/Replication/Manager/ReplicationManager.h
@@ -29,6 +29,8 @@ namespace UE::ConcertSyncClient::Replication
 		virtual void LeaveReplicationSession() override;
 		virtual bool CanJoin() override;
 		virtual bool IsConnectedToReplicationSession() override;
+		virtual EStreamEnumerationResult ForEachRegisteredStream(TFunctionRef<EBreakBehavior(const FReplicationStreamDescription& Stream)> Callback) const override;
+		virtual TFuture<FAuthorityChangeResponse> RequestAuthorityChange(FAuthorityChangeRequest Args) override;
 		//~ End IConcertClientReplicationManager Interface
 
 	private:
diff --git a/Engine/Plugins/Developer/Concert/ConcertSync/ConcertSyncClient/Source/ConcertSyncClient/Private/Replication/Manager/ReplicationManagerState_Connected.cpp b/Engine/Plugins/Developer/Concert/ConcertSync/ConcertSyncClient/Source/ConcertSyncClient/Private/Replication/Manager/ReplicationManagerState_Connected.cpp
--- a/Engine/Plugins/Developer/Concert/ConcertSync/ConcertSyncClient/Source/ConcertSyncClient/Private/Replication/Manager/ReplicationManagerState_Connected.cpp
+++ b/Engine/Plugins/Developer/Concert/ConcertSync/ConcertSyncClient/Source/ConcertSyncClient/Private/Replication/Manager/ReplicationManagerState_Connected.cpp
@@ -27,9 +27,10 @@ namespace UE::ConcertSyncClient::Replication
 		: FReplicationManagerState(Owner)
 		, LiveSession(LiveSession)
 		, ReplicationBridge(ReplicationBridge)
+		, RegisteredStreams(MoveTemp(StreamDescriptions))
 		// TODO: Use config to determine which replication format to use
 		, ReplicationFormat(MakeShared<ConcertSyncCore::FFullObjectFormat>())
-		, ReplicationDataSource(MakeShared<FClientReplicationDataCollector>(ReplicationBridge, ReplicationFormat, StreamDescriptions))
+		, ReplicationDataSource(MakeShared<FClientReplicationDataCollector>(ReplicationBridge, ReplicationFormat, RegisteredStreams))
 		, Sender(MakeShared<ConcertSyncCore::FObjectReplicationSender>(LiveSession->GetSessionServerEndpointId(), LiveSession, ReplicationDataSource))
 		, ReceivedDataCache(MakeShared<ConcertSyncCore::FObjectReplicationCache>(ReplicationFormat))
 		, Receiver(MakeShared<ConcertSyncCore::FObjectReplicationReceiver>(LiveSession, ReceivedDataCache))
@@ -55,6 +56,29 @@ namespace UE::ConcertSyncClient::Replication
 		ChangeState(MakeShared<FReplicationManagerState_Disconnected>(LiveSession, ReplicationBridge, GetOwner()));
 	}
 
+	IConcertClientReplicationManager::EStreamEnumerationResult FReplicationManagerState_Connected::ForEachRegisteredStream(
+		TFunctionRef<EBreakBehavior(const FReplicationStreamDescription& Stream)> Callback
+		) const
+	{
+		for (const FReplicationStreamDescription& Stream : RegisteredStreams)
+		{
+			if (Callback(Stream) == EBreakBehavior::Break)
+			{
+				break;
+			}
+		}
+		return EStreamEnumerationResult::Iterated;
+	}
+
+	TFuture<FAuthorityChangeResponse> FReplicationManagerState_Connected::RequestAuthorityChange(FAuthorityChangeRequest Args)
+	{
+		return LiveSession->SendCustomRequest<FConcertChangeAuthority_Request, FConcertChangeAuthority_Response>(Args, LiveSession->GetSessionServerEndpointId())
+			.Next([](FConcertChangeAuthority_Response&& Response)
+			{
+				return FAuthorityChangeResponse { MoveTemp(Response) };
+			});
+	}
+
 	void FReplicationManagerState_Connected::OnEnterState()
 	{
 		LiveSession->OnTick().AddSP(this, &FReplicationManagerState_Connected::Tick);
diff --git a/Engine/Plugins/Developer/Concert/ConcertSync/ConcertSyncClient/Source/ConcertSyncClient/Private/Replication/Manager/ReplicationManagerState_Connected.h b/Engine/Plugins/Developer/Concert/ConcertSync/ConcertSyncClient/Source/ConcertSyncClient/Private/Replication/Manager/ReplicationManagerState_Connected.h
--- a/Engine/Plugins/Developer/Concert/ConcertSync/ConcertSyncClient/Source/ConcertSyncClient/Private/Replication/Manager/ReplicationManagerState_Connected.h
+++ b/Engine/Plugins/Developer/Concert/ConcertSync/ConcertSyncClient/Source/ConcertSyncClient/Private/Replication/Manager/ReplicationManagerState_Connected.h
@@ -47,33 +47,37 @@ namespace UE::ConcertSyncClient::Replication
 		virtual void LeaveReplicationSession() override;
 		virtual bool CanJoin() override { return false; }
 		virtual bool IsConnectedToReplicationSession() override { return true; }
+		virtual EStreamEnumerationResult ForEachRegisteredStream(TFunctionRef<EBreakBehavior(const FReplicationStreamDescription& Stream)> Callback) const override;
+		virtual TFuture<FAuthorityChangeResponse> RequestAuthorityChange(FAuthorityChangeRequest Args) override;
 		//~ End IConcertClientReplicationManager Interface
 
 	private:
 
 		/** Passed to FReplicationManagerState_Disconnected */
-		TSharedRef<IConcertClientSession> LiveSession;
+		const TSharedRef<IConcertClientSession> LiveSession;
 		/** Passed to FReplicationManagerState_Disconnected */
-		IConcertClientReplicationBridge* ReplicationBridge;
+		IConcertClientReplicationBridge* const ReplicationBridge;
+		/** The streams that were registered during the handshake. */
+		const TArray<FReplicationStreamDescription> RegisteredStreams;
 		
 		/** The format this client will use for sending & receiving data. */
-		TSharedRef<ConcertSyncCore::IObjectReplicationFormat> ReplicationFormat;
+		const TSharedRef<ConcertSyncCore::IObjectReplicationFormat> ReplicationFormat;
 
 		// Sending
 		/** Used as source of replication data. */
-		TSharedRef<FClientReplicationDataCollector> ReplicationDataSource;
+		const TSharedRef<FClientReplicationDataCollector> ReplicationDataSource;
 		/** Sends data collected by ReplicationDataSource to the server. */
-		TSharedRef<ConcertSyncCore::FObjectReplicationSender> Sender;
+		const TSharedRef<ConcertSyncCore::FObjectReplicationSender> Sender;
 
 		// Receiving
 		/** Stores data received by Receiver until it is consumed by ReceivedReplicationQueuer. */
-		TSharedRef<ConcertSyncCore::FObjectReplicationCache> ReceivedDataCache;
+		const TSharedRef<ConcertSyncCore::FObjectReplicationCache> ReceivedDataCache;
 		/** Receives data from remote endpoints via message bus.  */
-		TSharedRef<ConcertSyncCore::FObjectReplicationReceiver> Receiver;
+		const TSharedRef<ConcertSyncCore::FObjectReplicationReceiver> Receiver;
 		/** Queues data until is can be processed. */
-		TSharedRef<FClientReplicationDataQueuer> ReceivedReplicationQueuer;
+		const TSharedRef<FClientReplicationDataQueuer> ReceivedReplicationQueuer;
 		/** Processes data from ReceivedReplicationQueuer once we tick. */
-		TSharedRef<FObjectReplicationApplierProcessor> ReplicationApplier;
+		const TSharedRef<FObjectReplicationApplierProcessor> ReplicationApplier;
 
 		using FTickTask = void(FReplicationManagerState_Connected::*)(float TimeBudget);
 		/**
diff --git a/Engine/Plugins/Developer/Concert/ConcertSync/ConcertSyncClient/Source/ConcertSyncClient/Private/Replication/Manager/ReplicationManagerState_Disconnected.h b/Engine/Plugins/Developer/Concert/ConcertSync/ConcertSyncClient/Source/ConcertSyncClient/Private/Replication/Manager/ReplicationManagerState_Disconnected.h
--- a/Engine/Plugins/Developer/Concert/ConcertSync/ConcertSyncClient/Source/ConcertSyncClient/Private/Replication/Manager/ReplicationManagerState_Disconnected.h
+++ b/Engine/Plugins/Developer/Concert/ConcertSync/ConcertSyncClient/Source/ConcertSyncClient/Private/Replication/Manager/ReplicationManagerState_Disconnected.h
@@ -3,6 +3,7 @@
 #pragma once
 
 #include "ReplicationManagerState.h"
+#include "ReplicationManagerUtils.h"
 #include "Replication/Processing/ClientReplicationDataCollector.h"
 
 class IConcertClientSession;
@@ -25,6 +26,8 @@ namespace UE::ConcertSyncClient::Replication
 		virtual void LeaveReplicationSession() override;
 		virtual bool CanJoin() override { return true; }
 		virtual bool IsConnectedToReplicationSession() override { return false; }
+		virtual EStreamEnumerationResult ForEachRegisteredStream(TFunctionRef<EBreakBehavior(const FReplicationStreamDescription& Stream)> Callback) const override { return EStreamEnumerationResult::NoRegisteredStreams; }
+		virtual TFuture<FAuthorityChangeResponse> RequestAuthorityChange(FAuthorityChangeRequest Args) override { return RejectAll(MoveTemp(Args)); }
 		//~ End IConcertClientReplicationManager Interface
 
 	private:
diff --git a/Engine/Plugins/Developer/Concert/ConcertSync/ConcertSyncClient/Source/ConcertSyncClient/Private/Replication/Manager/ReplicationManagerState_Handshaking.h b/Engine/Plugins/Developer/Concert/ConcertSync/ConcertSyncClient/Source/ConcertSyncClient/Private/Replication/Manager/ReplicationManagerState_Handshaking.h
--- a/Engine/Plugins/Developer/Concert/ConcertSync/ConcertSyncClient/Source/ConcertSyncClient/Private/Replication/Manager/ReplicationManagerState_Handshaking.h
+++ b/Engine/Plugins/Developer/Concert/ConcertSync/ConcertSyncClient/Source/ConcertSyncClient/Private/Replication/Manager/ReplicationManagerState_Handshaking.h
@@ -4,6 +4,7 @@
 
 #include "IConcertSessionHandler.h"
 #include "ReplicationManagerState.h"
+#include "ReplicationManagerUtils.h"
 #include "Replication/Messages/ConcertReplicationHandshakeMessages.h"
 #include "Replication/Processing/ClientReplicationDataCollector.h"
 
@@ -30,6 +31,8 @@ namespace UE::ConcertSyncClient::Replication
 		virtual void LeaveReplicationSession() override;
 		virtual bool CanJoin() override { return false; }
 		virtual bool IsConnectedToReplicationSession() override { return false; }
+		virtual EStreamEnumerationResult ForEachRegisteredStream(TFunctionRef<EBreakBehavior(const FReplicationStreamDescription& Stream)> Callback) const override { return EStreamEnumerationResult::NoRegisteredStreams; }
+		virtual TFuture<FAuthorityChangeResponse> RequestAuthorityChange(FAuthorityChangeRequest Args) override { return RejectAll(MoveTemp(Args)); }
 		//~ End IConcertClientReplicationManager Interface
 
 	private:
diff --git a/Engine/Plugins/Developer/Concert/ConcertSync/ConcertSyncClient/Source/ConcertSyncClient/Private/Replication/Manager/ReplicationManagerUtils.cpp b/Engine/Plugins/Developer/Concert/ConcertSync/ConcertSyncClient/Source/ConcertSyncClient/Private/Replication/Manager/ReplicationManagerUtils.cpp
new file mode 100644
--- /dev/null
+++ b/Engine/Plugins/Developer/Concert/ConcertSync/ConcertSyncClient/Source/ConcertSyncClient/Private/Replication/Manager/ReplicationManagerUtils.cpp
@@ -0,0 +1,18 @@
+// Copyright Epic Games, Inc. All Rights Reserved.
+
+#include "ReplicationManagerUtils.h"
+
+#include "Replication/IConcertClientReplicationManager.h"
+
+#include "Algo/Transform.h"
+#include "Containers/Set.h"
+
+namespace UE::ConcertSyncClient::Replication
+{
+	TFuture<FAuthorityChangeResponse> RejectAll(FAuthorityChangeRequest&& Args)
+	{
+		return MakeFulfilledPromise<FAuthorityChangeResponse>(
+			FAuthorityChangeResponse{{ .RejectedObjects = MoveTemp(Args.TakeAuthority) }}
+			).GetFuture();
+	}
+}
diff --git a/Engine/Plugins/Developer/Concert/ConcertSync/ConcertSyncClient/Source/ConcertSyncClient/Private/Replication/Manager/ReplicationManagerUtils.h b/Engine/Plugins/Developer/Concert/ConcertSync/ConcertSyncClient/Source/ConcertSyncClient/Private/Replication/Manager/ReplicationManagerUtils.h
new file mode 100644
--- /dev/null
+++ b/Engine/Plugins/Developer/Concert/ConcertSync/ConcertSyncClient/Source/ConcertSyncClient/Private/Replication/Manager/ReplicationManagerUtils.h
@@ -0,0 +1,14 @@
+// Copyright Epic Games, Inc. All Rights Reserved.
+
+#pragma once
+
+#include "Async/Future.h"
+
+namespace UE::ConcertSyncClient::Replication
+{
+	struct FAuthorityChangeRequest;
+	struct FAuthorityChangeResponse;
+
+	/** Creates a fullfilled future which rejects all items in FConcertChangeAuthority_Request::TakeAuthority. */
+	TFuture<FAuthorityChangeResponse> RejectAll(FAuthorityChangeRequest&& Args);
+};
diff --git a/Engine/Plugins/Developer/Concert/ConcertSync/ConcertSyncClient/Source/ConcertSyncClient/Private/Replication/Processing/ClientReplicationDataCollector.cpp b/Engine/Plugins/Developer/Concert/ConcertSync/ConcertSyncClient/Source/ConcertSyncClient/Private/Replication/Processing/ClientReplicationDataCollector.cpp
--- a/Engine/Plugins/Developer/Concert/ConcertSync/ConcertSyncClient/Source/ConcertSyncClient/Private/Replication/Processing/ClientReplicationDataCollector.cpp
+++ b/Engine/Plugins/Developer/Concert/ConcertSync/ConcertSyncClient/Source/ConcertSyncClient/Private/Replication/Processing/ClientReplicationDataCollector.cpp
@@ -13,7 +13,7 @@ namespace UE::ConcertSyncClient::Replication
 	FClientReplicationDataCollector::FClientReplicationDataCollector(
 		IConcertClientReplicationBridge* ReplicationBridge,
 		TSharedRef<ConcertSyncCore::IObjectReplicationFormat> ReplicationFormat,
-		TArrayView<FReplicationStreamDescription> StreamDescriptions
+		TArrayView<const FReplicationStreamDescription> StreamDescriptions
 		)
 		: Bridge(ReplicationBridge)
 		, ReplicationFormat(MoveTemp(ReplicationFormat))
diff --git a/Engine/Plugins/Developer/Concert/ConcertSync/ConcertSyncClient/Source/ConcertSyncClient/Private/Replication/Processing/ClientReplicationDataCollector.h b/Engine/Plugins/Developer/Concert/ConcertSync/ConcertSyncClient/Source/ConcertSyncClient/Private/Replication/Processing/ClientReplicationDataCollector.h
--- a/Engine/Plugins/Developer/Concert/ConcertSync/ConcertSyncClient/Source/ConcertSyncClient/Private/Replication/Processing/ClientReplicationDataCollector.h
+++ b/Engine/Plugins/Developer/Concert/ConcertSync/ConcertSyncClient/Source/ConcertSyncClient/Private/Replication/Processing/ClientReplicationDataCollector.h
@@ -36,7 +36,7 @@ namespace UE::ConcertSyncClient::Replication
 		FClientReplicationDataCollector(
 			IConcertClientReplicationBridge* ReplicationBridge,
 			TSharedRef<ConcertSyncCore::IObjectReplicationFormat> ReplicationFormat,
-			TArrayView<FReplicationStreamDescription> StreamsToSend
+			TArrayView<const FReplicationStreamDescription> StreamsToSend
 			);
 		virtual ~FClientReplicationDataCollector() override;
 
diff --git a/Engine/Plugins/Developer/Concert/ConcertSync/ConcertSyncClient/Source/ConcertSyncClient/Public/Replication/IConcertClientReplicationManager.h b/Engine/Plugins/Developer/Concert/ConcertSync/ConcertSyncClient/Source/ConcertSyncClient/Public/Replication/IConcertClientReplicationManager.h
--- a/Engine/Plugins/Developer/Concert/ConcertSync/ConcertSyncClient/Source/ConcertSyncClient/Public/Replication/IConcertClientReplicationManager.h
+++ b/Engine/Plugins/Developer/Concert/ConcertSync/ConcertSyncClient/Source/ConcertSyncClient/Public/Replication/IConcertClientReplicationManager.h
@@ -2,8 +2,10 @@
 
 #pragma once
 
+#include "Misc/EBreakBehavior.h"
 #include "Replication/Messages/ConcertReplicationHandshakeMessages.h"
 #include "Replication/Data/ReplicationClientDescription.h"
+#include "Replication/Messages/ConcertReplicationEvents.h"
 
 template<typename ResultType>
 class TFuture;
@@ -33,6 +35,12 @@ namespace UE::ConcertSyncClient::Replication
 			, DetailedErrorMessage(MoveTemp(DetailedErrorMessage))
 		{}
 	};
+
+	// The intention here is to wrap the request in case there are some more specific meta data we want to add in the future.
+	struct FAuthorityChangeRequest : FConcertChangeAuthority_Request
+	{};
+	struct FAuthorityChangeResponse : FConcertChangeAuthority_Response
+	{};
 }
 
 /**
@@ -41,7 +49,7 @@ namespace UE::ConcertSyncClient::Replication
  * Keeps a list of properties to send along with their send rules.
  * Tells the server which properties this client is interested in receiving.
  */
-class IConcertClientReplicationManager
+class CONCERTSYNCCLIENT_API IConcertClientReplicationManager
 {
 public:
 
@@ -52,7 +60,6 @@ public:
 	 * @note This may execute on any thread. Take care to synchronize correctly with the game thread if needed.
 	 */
 	virtual TFuture<UE::ConcertSyncClient::Replication::FJoinReplicatedSessionResult> JoinReplicationSession(UE::ConcertSyncClient::Replication::FJoinReplicatedSessionArgs Args) = 0;
-
 	/** Leaves the current replication session. */
 	virtual void LeaveReplicationSession() = 0;
 
@@ -61,9 +68,29 @@ public:
 	 * If this returns true, you can call JoinReplicationSession to attempt joining the session.
 	 */
 	virtual bool CanJoin() = 0;
-		
 	/** Whether JoinReplicationSession completed successfully and LeaveReplicationSession has not yet been called. */
 	virtual bool IsConnectedToReplicationSession() = 0;
-		
+
+	enum class EStreamEnumerationResult { NoRegisteredStreams, Iterated };
+	/**
+	 * Iterates the streams that were registered in the JoinReplicationSession request.
+	 * It only makes sense to call this function the manager has joined a replication session.
+	 * @return Whether this manager is connected to a session (Iterated) or not (NoRegisteredStreams).
+	 */
+	virtual EStreamEnumerationResult ForEachRegisteredStream(TFunctionRef<EBreakBehavior(const FReplicationStreamDescription& Stream)> Callback) const = 0;
+	/** @return Whether this manager is currently in a replication session (basically whether ForEachRegisteredStream returns EStreamEnumerationResult::Iterated). */
+	bool HasRegisteredStreams() const;
+	
+	/**
+	 * Requests from the server to change the authority over some objects.
+	 * @note This may execute on any thread. Take care of synchronize correctly with the game thread if needed.
+	 */
+	virtual TFuture<UE::ConcertSyncClient::Replication::FAuthorityChangeResponse> RequestAuthorityChange(UE::ConcertSyncClient::Replication::FAuthorityChangeRequest Args) = 0;
+	/** Util function that will request authority for all streams for the given objects. */
+	TFuture<UE::ConcertSyncClient::Replication::FAuthorityChangeResponse> TakeAuthorityOver(TArrayView<const FSoftObjectPath> Objects);
+	/** Util function that will let go over all authority of the given objects. */
+	TFuture<UE::ConcertSyncClient::Replication::FAuthorityChangeResponse> ReleaseAuthorityOf(TArrayView<const FSoftObjectPath> Objects);
+	
 	virtual ~IConcertClientReplicationManager() = default;
-};
\ No newline at end of file
+};
+	
\ No newline at end of file
diff --git a/Engine/Plugins/Developer/Concert/ConcertSync/ConcertSyncCore/Source/ConcertSyncCore/Private/Replication/Data/ConcertPropertySelection.cpp b/Engine/Plugins/Developer/Concert/ConcertSync/ConcertSyncCore/Source/ConcertSyncCore/Private/Replication/Data/ConcertPropertySelection.cpp
--- a/Engine/Plugins/Developer/Concert/ConcertSync/ConcertSyncCore/Source/ConcertSyncCore/Private/Replication/Data/ConcertPropertySelection.cpp
+++ b/Engine/Plugins/Developer/Concert/ConcertSync/ConcertSyncCore/Source/ConcertSyncCore/Private/Replication/Data/ConcertPropertySelection.cpp
@@ -9,7 +9,7 @@
 
 const FName FConcertPropertyChain::InternalContainerPropertyValueName(TEXT("Value"));
 
-TOptional<FConcertPropertyChain> FConcertPropertyChain::CreateFromPath(UStruct& Class, const TArray<FName>& NamePath)
+TOptional<FConcertPropertyChain> FConcertPropertyChain::CreateFromPath(const UStruct& Class, const TArray<FName>& NamePath)
 {
 	TOptional<FConcertPropertyChain> Result;
 	// The goal here is to find a valid FProperty path based on the NamePath...
@@ -202,6 +202,25 @@ FString FConcertPropertyChain::ToString(EToStringMethod Method) const
 	
 }
 
+bool FConcertPropertySelection::OverlapsWith(const FConcertPropertySelection& Other) const
+{
+	/* 
+	 * Performance could be improved: Two chains can only overlap if their root properties overlap.
+	 * Example: [Foo.Vector.X] and [Foo.FloatProperty] both share Foo struct (but do not overlap in this case).
+	 */
+	for (const FConcertPropertyChain& ThisChain : ReplicatedProperties)
+	{
+		for (const FConcertPropertyChain& OtherChain : Other.ReplicatedProperties)
+		{
+			if (ThisChain == OtherChain) 
+			{
+				return true;
+			}
+		}
+	}
+	return false;
+}
+
 uint32 GetTypeHash(const FConcertPropertyChain& Chain)
 {
 	return GetTypeHash(Chain.GetPathToProperty());
diff --git a/Engine/Plugins/Developer/Concert/ConcertSync/ConcertSyncCore/Source/ConcertSyncCore/Private/Replication/Messages/ConcertReplicationHandshakeMessages.cpp b/Engine/Plugins/Developer/Concert/ConcertSync/ConcertSyncCore/Source/ConcertSyncCore/Private/Replication/Messages/ConcertReplicationHandshakeMessages.cpp
--- a/Engine/Plugins/Developer/Concert/ConcertSync/ConcertSyncCore/Source/ConcertSyncCore/Private/Replication/Messages/ConcertReplicationHandshakeMessages.cpp
+++ b/Engine/Plugins/Developer/Concert/ConcertSync/ConcertSyncCore/Source/ConcertSyncCore/Private/Replication/Messages/ConcertReplicationHandshakeMessages.cpp
@@ -6,7 +6,7 @@ namespace UE::ConcertSyncCore::Replication
 {
 	FString LexJoinErrorCode(EJoinReplicationErrorCode ErrorCode)
 	{
-		static_assert(static_cast<int32>(EJoinReplicationErrorCode::Max) == 10, "Update JoinErrorCodeToText when you change entries");
+		static_assert(static_cast<int32>(EJoinReplicationErrorCode::Max) == 9, "Update JoinErrorCodeToText when you change entries");
 
 		switch (ErrorCode)
 		{
@@ -18,8 +18,7 @@ namespace UE::ConcertSyncCore::Replication
 		case EJoinReplicationErrorCode::InvalidClass: return TEXT("InvalidClass");
 		case EJoinReplicationErrorCode::DuplicateProperty: return TEXT("DuplicateProperty");
 		case EJoinReplicationErrorCode::AlreadyInSession: return TEXT("AlreadyInSession");
-		case EJoinReplicationErrorCode::ConflictingStreamId: return TEXT("ConflictingStreamId");
-		case EJoinReplicationErrorCode::ConflictingAuthority: return TEXT("ConflictingAuthority");
+		case EJoinReplicationErrorCode::DuplicateStreamId: return TEXT("DuplicateStreamId");
 		case EJoinReplicationErrorCode::FailedToUnpackStream: return TEXT("FailedToUnpackStream");
 		default: 
 			return FString();
diff --git a/Engine/Plugins/Developer/Concert/ConcertSync/ConcertSyncCore/Source/ConcertSyncCore/Private/Replication/Processing/ObjectReplicationReceiver.cpp b/Engine/Plugins/Developer/Concert/ConcertSync/ConcertSyncCore/Source/ConcertSyncCore/Private/Replication/Processing/ObjectReplicationReceiver.cpp
--- a/Engine/Plugins/Developer/Concert/ConcertSync/ConcertSyncCore/Source/ConcertSyncCore/Private/Replication/Processing/ObjectReplicationReceiver.cpp
+++ b/Engine/Plugins/Developer/Concert/ConcertSync/ConcertSyncCore/Source/ConcertSyncCore/Private/Replication/Processing/ObjectReplicationReceiver.cpp
@@ -7,8 +7,12 @@
 #include "Replication/Messages/ConcertReplicationEvents.h"
 #include "Replication/Processing/ObjectReplicationCache.h"
 
+#include "HAL/IConsoleManager.h"
+
 namespace UE::ConcertSyncCore
 {
+	static TAutoConsoleVariable<bool> CVarLogReceivedObjects(TEXT("Concert.Replication.LogReceivedObjects"), false, TEXT("Enable Concert logging for received replicated objects."));
+	
 	FObjectReplicationReceiver::FObjectReplicationReceiver(TSharedRef<IConcertSession> Session, TSharedRef<FObjectReplicationCache> ReplicationCache)
 		: Session(MoveTemp(Session))
 		, ReplicationCache(MoveTemp(ReplicationCache))
@@ -23,26 +27,38 @@ namespace UE::ConcertSyncCore
 
 	void FObjectReplicationReceiver::HandleBatchReplicationEvent(const FConcertSessionContext& SessionContext, const FConcertBatchReplicationEvent& Event)
 	{
+		// Fyi: an object may have multiple changes in a batch replication event: each stream can modify different properties as long as they do not overlap.
 		int32 NumObjects = 0;
+		int32 NumRejectedObjectChanges = 0;
 		int32 NumCacheUsages = 0;
-		int32 NumOfAcceptedObjects = 0;
+		int32 NumOfAcceptedObjectChanges = 0;
 		
 		for (const FConcertStreamReplicationEvent& StreamEvent : Event.Streams)
 		{
-			NumObjects += StreamEvent.ReplicatedObjects.Num();
+			const int32 ObjectsInStream = StreamEvent.ReplicatedObjects.Num();
+			NumObjects += ObjectsInStream;
+			
 			for (const FConcertObjectReplicationEvent& ObjectEvent : StreamEvent.ReplicatedObjects)
 			{
-				const int32 NumAccepted = ReplicationCache->StoreUntilConsumed(SessionContext.SourceEndpointId, StreamEvent.StreamId, ObjectEvent);
-				NumCacheUsages += NumAccepted;
-				NumOfAcceptedObjects += NumAccepted == 0 ? 0 : 1;
+				if (ShouldAcceptObject(SessionContext, StreamEvent, ObjectEvent))
+				{
+					const int32 NumAccepted = ReplicationCache->StoreUntilConsumed(SessionContext.SourceEndpointId, StreamEvent.StreamId, ObjectEvent);
+					NumCacheUsages += NumAccepted;
+					NumOfAcceptedObjectChanges += NumAccepted == 0 ? 0 : 1;
+				}
+				else
+				{
+					++NumRejectedObjectChanges;
+				}
 			}
 		}
 		
-		UE_LOG(LogConcert, Verbose, TEXT("Received %d streams with %d objects from endpoint %s. Cached %d objects with a total of %d cache usages."),
+		UE_CLOG(CVarLogReceivedObjects.GetValueOnGameThread(), LogConcert, Log, TEXT("Received %d streams with %d object changes from endpoint %s. Rejected %d object changes. Cached %d object changes with a total of %d cache usages."),
 			Event.Streams.Num(),
 			NumObjects,
 			*SessionContext.SourceEndpointId.ToString(),
-			NumOfAcceptedObjects,
+			NumRejectedObjectChanges,
+			NumOfAcceptedObjectChanges,
 			NumCacheUsages
 			);
 	}
diff --git a/Engine/Plugins/Developer/Concert/ConcertSync/ConcertSyncCore/Source/ConcertSyncCore/Private/Replication/Processing/ObjectReplicationSender.cpp b/Engine/Plugins/Developer/Concert/ConcertSync/ConcertSyncCore/Source/ConcertSyncCore/Private/Replication/Processing/ObjectReplicationSender.cpp
--- a/Engine/Plugins/Developer/Concert/ConcertSync/ConcertSyncCore/Source/ConcertSyncCore/Private/Replication/Processing/ObjectReplicationSender.cpp
+++ b/Engine/Plugins/Developer/Concert/ConcertSync/ConcertSyncCore/Source/ConcertSyncCore/Private/Replication/Processing/ObjectReplicationSender.cpp
@@ -4,11 +4,15 @@
 
 #include "ConcertLogGlobal.h"
 #include "IConcertSession.h"
-#include "Algo/Accumulate.h"
 #include "Replication/Processing/IReplicationDataSource.h"
 
+#include "Algo/Accumulate.h"
+#include "HAL/IConsoleManager.h"
+
 namespace UE::ConcertSyncCore
 {
+	static TAutoConsoleVariable<bool> CVarLogSentObjects(TEXT("Concert.Replication.LogSentObjects"), false, TEXT("Enable Concert logging for sent replicated objects."));
+	
 	FObjectReplicationSender::FObjectReplicationSender(
 		const FGuid& TargetEndpointId,
 		TSharedRef<IConcertSession> Session,
@@ -26,7 +30,7 @@ namespace UE::ConcertSyncCore
 		if (!EventToSend.Streams.IsEmpty())
 		{
 			const int32 NumObjects = Algo::TransformAccumulate(EventToSend.Streams, [](const FConcertStreamReplicationEvent& Event){ return Event.ReplicatedObjects.Num(); }, 0);
-			UE_LOG(LogConcert, Verbose, TEXT("Sending %d streams with %d objects to %s"),
+			UE_CLOG(CVarLogSentObjects.GetValueOnGameThread(), LogConcert, Log, TEXT("Sending %d streams with %d objects to %s"),
 				EventToSend.Streams.Num(),
 				NumObjects,
 				*TargetEndpointId.ToString()
diff --git a/Engine/Plugins/Developer/Concert/ConcertSync/ConcertSyncCore/Source/ConcertSyncCore/Public/Replication/Data/ConcertPropertySelection.h b/Engine/Plugins/Developer/Concert/ConcertSync/ConcertSyncCore/Source/ConcertSyncCore/Public/Replication/Data/ConcertPropertySelection.h
--- a/Engine/Plugins/Developer/Concert/ConcertSync/ConcertSyncCore/Source/ConcertSyncCore/Public/Replication/Data/ConcertPropertySelection.h
+++ b/Engine/Plugins/Developer/Concert/ConcertSync/ConcertSyncCore/Source/ConcertSyncCore/Public/Replication/Data/ConcertPropertySelection.h
@@ -32,7 +32,7 @@ struct CONCERTSYNCCORE_API FConcertPropertyChain
 	static const FName InternalContainerPropertyValueName;
 
 	/** Constructs a FConcertPropertyChain from a path if it is valid. If you need to create many paths in one go, use PropertyUtils::BulkConstructConcertChainsFromPaths instead. */
-	static TOptional<FConcertPropertyChain> CreateFromPath(UStruct& Class, const TArray<FName>& NamePath);
+	static TOptional<FConcertPropertyChain> CreateFromPath(const UStruct& Class, const TArray<FName>& NamePath);
 
 	FConcertPropertyChain() = default;
 	/**
@@ -42,8 +42,10 @@ struct CONCERTSYNCCORE_API FConcertPropertyChain
 	FConcertPropertyChain(const FArchiveSerializedPropertyChain* OptionalChain, const FProperty& LeafProperty);
 	
 	/** Gets the leaf property, which is the property the path leads towards. */
-	FName GetLeafProperty() const { return PathToProperty.IsEmpty() ? NAME_None : PathToProperty[PathToProperty.Num() - 1]; }
+	FName GetLeafProperty() const { return IsEmpty() ? NAME_None : PathToProperty[PathToProperty.Num() - 1]; }
+	FName GetRootProperty() const { return IsEmpty() ? NAME_None : PathToProperty[0]; }
 	bool IsRootProperty() const { return PathToProperty.Num() == 1; }
+	bool IsEmpty() const { return PathToProperty.IsEmpty(); }
 
 	/** @return Whether this is a parent of ChildToCheck */
 	bool IsParentOf(const FConcertPropertyChain& ChildToCheck) const { return ChildToCheck.IsChildOf(*this); }
@@ -178,6 +180,9 @@ struct CONCERTSYNCCORE_API FConcertPropertySelection
 	 */
 	UPROPERTY()
 	TArray<FConcertPropertyChain> ReplicatedProperties;
+
+	/** @return Whether this and Other contain at least one property that is the same. This algorithm is strictly O(n^2) but runs O(n) on average. */
+	bool OverlapsWith(const FConcertPropertySelection& Other) const;
 };
 
 CONCERTSYNCCORE_API uint32 GetTypeHash(const FConcertPropertyChain& Chain);
\ No newline at end of file
diff --git a/Engine/Plugins/Developer/Concert/ConcertSync/ConcertSyncCore/Source/ConcertSyncCore/Public/Replication/Messages/ConcertReplicationEvents.h b/Engine/Plugins/Developer/Concert/ConcertSync/ConcertSyncCore/Source/ConcertSyncCore/Public/Replication/Messages/ConcertReplicationEvents.h
--- a/Engine/Plugins/Developer/Concert/ConcertSync/ConcertSyncCore/Source/ConcertSyncCore/Public/Replication/Messages/ConcertReplicationEvents.h
+++ b/Engine/Plugins/Developer/Concert/ConcertSync/ConcertSyncCore/Source/ConcertSyncCore/Public/Replication/Messages/ConcertReplicationEvents.h
@@ -31,7 +31,7 @@ struct FConcertStreamReplicationEvent
 
 	/** The stream that produced this data. */
 	UPROPERTY()
-	FGuid StreamId;
+	FGuid StreamId; // TODO DP UE-193261: This field is only relevant to the server and useless for receiving clients.
 
 	/** The objects replicated by this event */
 	UPROPERTY()
@@ -48,13 +48,69 @@ struct FConcertStreamReplicationEvent
  * Sent from
  *  - Authoritative client to server
  *  - Server to clients. Every event may be customized for each client (clients can choose to ignore object updates).
+ *
+ *  TODO DP UE-193261: This should be split up into FConcertBatchReplicationEvent_ToServer and FConcertBatchReplicationEvent_ToClient; clients do not need FConcertStreamReplicationEvent::StreamId. 
  */
 USTRUCT()
-struct FConcertBatchReplicationEvent
+struct FConcertBatchReplicationEvent 
 {
 	GENERATED_BODY()
 	
 	/** The objects replicated by this event */
 	UPROPERTY()
 	TArray<FConcertStreamReplicationEvent> Streams;
+};
+
+USTRUCT()
+struct FConcertStreamArray
+{
+	GENERATED_BODY()
+
+	UPROPERTY()
+	TArray<FGuid> StreamIds;
+};
+
+/**
+ * Sent by clients to the server to attempt to take or release authority over properties on objects.
+ * Clients must have authority over object properties before sending FConcertBatchReplicationEvent, or the server will reject the updates.
+ * 
+ * Multiple clients can have authority over the same objects as long as the properties do not overlap.
+ * The properties are defined implicitly by providing the stream ID that will be writing the properties.
+ */
+USTRUCT()
+struct FConcertChangeAuthority_Request
+{
+	GENERATED_BODY()
+
+	// In the future we can consider adding an option to make this all or nothing (if one object fails, the entire request does)
+	// It is supposed to be a convenience and network optimization to send one big request instead of many small ones.
+
+	/**
+	 * Objects the client requests to take authority over.
+	 * Mapping of object to stream identifiers the client has previously registered.
+	 */
+	UPROPERTY()
+	TMap<FSoftObjectPath, FConcertStreamArray> TakeAuthority;
+
+	/**
+	 * Objects the client no longer needs authority over.
+	 * Mapping of object to stream identifiers the client has previously registered.
+	 * 
+	 * The request will not fail if it contains objects it has no authority over.
+	 */
+	UPROPERTY()
+	TMap<FSoftObjectPath, FConcertStreamArray> ReleaseAuthority;
+};
+
+USTRUCT()
+struct FConcertChangeAuthority_Response
+{
+	GENERATED_BODY()
+
+	/**
+	 * Objects streams the client did not receive authority over.
+	 * The client is implied to have authority over all other request objects.
+	 */
+	UPROPERTY()
+	TMap<FSoftObjectPath, FConcertStreamArray> RejectedObjects;
 };
\ No newline at end of file
diff --git a/Engine/Plugins/Developer/Concert/ConcertSync/ConcertSyncCore/Source/ConcertSyncCore/Public/Replication/Messages/ConcertReplicationHandshakeMessages.h b/Engine/Plugins/Developer/Concert/ConcertSync/ConcertSyncCore/Source/ConcertSyncCore/Public/Replication/Messages/ConcertReplicationHandshakeMessages.h
--- a/Engine/Plugins/Developer/Concert/ConcertSync/ConcertSyncCore/Source/ConcertSyncCore/Public/Replication/Messages/ConcertReplicationHandshakeMessages.h
+++ b/Engine/Plugins/Developer/Concert/ConcertSync/ConcertSyncCore/Source/ConcertSyncCore/Public/Replication/Messages/ConcertReplicationHandshakeMessages.h
@@ -23,28 +23,24 @@ enum class EJoinReplicationErrorCode : uint8
 
 	// Server error codes:
 
-	/** Rare case in which client is no longer in a session but has requested to join a replication session. */
-	NotInAnyConcertSession = 4,
-	
-	/** One of the specified classes was null or otherwise invalid */
-	InvalidClass = 5,
-	
-	/** An object selection contained a property twice. */
-	DuplicateProperty = 6,
-	
+	// Session:
 	/**
-	 * Client has already joined a replication session. Send FConcertLeaveReplicationSessionEvent first.
+	 * Client is not in any session but has requested to join a replication session.
+	 * This case should never happen under normal conditions e.g. when using Multi User Browser, but malicious users could make invalid requests.
 	 */
-	AlreadyInSession = 7,
-
-	/** You or another client already is using the same stream id. */
-	ConflictingStreamId = 8,
-	
-	/** You or another client already has authority over some of the data. */
-	ConflictingAuthority = 9,
+	NotInAnyConcertSession = 4,
+	/** Client has already joined a replication session. Send FConcertLeaveReplicationSessionEvent first. */
+	AlreadyInSession = 5,
 
+	// Invalid input:
+	/** One of the specified classes was null or otherwise invalid */
+	InvalidClass = 6,
+	/** An object selection contained a property twice. */
+	DuplicateProperty = 7,
+	/** Your input contained the same stream id twice. */
+	DuplicateStreamId = 8,
 	/** Failed to unpack the stream, likely because one of the attribute classes could not be resolved. */
-	FailedToUnpackStream = 10,
+	FailedToUnpackStream = 9,
 	
 	MaxPlusOne,
 	/** The maximum possible entry value */
diff --git a/Engine/Plugins/Developer/Concert/ConcertSync/ConcertSyncCore/Source/ConcertSyncCore/Public/Replication/Processing/ObjectReplicationReceiver.h b/Engine/Plugins/Developer/Concert/ConcertSync/ConcertSyncCore/Source/ConcertSyncCore/Public/Replication/Processing/ObjectReplicationReceiver.h
--- a/Engine/Plugins/Developer/Concert/ConcertSync/ConcertSyncCore/Source/ConcertSyncCore/Public/Replication/Processing/ObjectReplicationReceiver.h
+++ b/Engine/Plugins/Developer/Concert/ConcertSync/ConcertSyncCore/Source/ConcertSyncCore/Public/Replication/Processing/ObjectReplicationReceiver.h
@@ -8,6 +8,7 @@ class IConcertSession;
 struct FConcertBatchReplicationEvent;
 struct FConcertObjectReplicationEvent;
 struct FConcertSessionContext;
+struct FConcertStreamReplicationEvent;
 
 namespace UE::ConcertSyncCore
 {
@@ -19,7 +20,12 @@ namespace UE::ConcertSyncCore
 	public:
 
 		FObjectReplicationReceiver(TSharedRef<IConcertSession> Session, TSharedRef<FObjectReplicationCache> ReplicationCache);
-		~FObjectReplicationReceiver();
+		virtual ~FObjectReplicationReceiver();
+
+	protected:
+
+		/** Whether the object should be processed. */
+		virtual bool ShouldAcceptObject(const FConcertSessionContext& SessionContext, const FConcertStreamReplicationEvent& StreamEvent, const FConcertObjectReplicationEvent& ObjectEvent) const { return true; }
 
 	private:
 
diff --git a/Engine/Plugins/Developer/Concert/ConcertSync/ConcertSyncServer/Source/ConcertSyncServer/Private/Replication/AuthorityManager.cpp b/Engine/Plugins/Developer/Concert/ConcertSync/ConcertSyncServer/Source/ConcertSyncServer/Private/Replication/AuthorityManager.cpp
new file mode 100644
--- /dev/null
+++ b/Engine/Plugins/Developer/Concert/ConcertSync/ConcertSyncServer/Source/ConcertSyncServer/Private/Replication/AuthorityManager.cpp
@@ -0,0 +1,207 @@
+// Copyright Epic Games, Inc. All Rights Reserved.
+
+#include "AuthorityManager.h"
+
+#include "ConcertLogGlobal.h"
+#include "IConcertSession.h"
+#include "Replication/ReplicationStreamObjectID.h"
+#include "Replication/Data/ReplicationStreamDescription.h"
+#include "Replication/Messages/ConcertReplicationEvents.h"
+
+namespace UE::ConcertSyncServer::Replication
+{
+	namespace Private
+	{
+		static void ForEachReplicatedObject(
+			const TMap<FSoftObjectPath, FConcertStreamArray>& Map,
+			TFunctionRef<void(const FGuid& StreamId, const FSoftObjectPath& ObjectPath)> Callback
+			)
+		{
+			for (const TPair<FSoftObjectPath, FConcertStreamArray>& ObjectInfo : Map)
+			{
+				for (const FGuid& StreamId : ObjectInfo.Value.StreamIds)
+				{
+					Callback(StreamId, ObjectInfo.Key);
+				}
+			}
+		}
+	}
+	
+	FAuthorityManager::FAuthorityManager(IAuthorityManagerGetters& Getters, TSharedRef<IConcertSession> InSession)
+		: Getters(Getters)
+		, Session(MoveTemp(InSession))
+	{
+		Session->RegisterCustomRequestHandler<FConcertChangeAuthority_Request, FConcertChangeAuthority_Response>(this, &FAuthorityManager::HandleChangeAuthorityRequest);
+	}
+
+	FAuthorityManager::~FAuthorityManager()
+	{
+		Session->UnregisterCustomRequestHandler<FConcertChangeAuthority_Request>();
+	}
+
+	bool FAuthorityManager::IsObjectChangeAllowed(const ConcertSyncCore::FReplicatedObjectInfo& ObjectChange) const
+	{
+		const FClientAuthorityData* AuthorityData = ClientAuthorityData.Find(ObjectChange.SenderEndpointId);
+		const TSet<FSoftObjectPath>* OwnedObjects = AuthorityData ? AuthorityData->OwnedObjects.Find(ObjectChange.StreamId) : nullptr;
+		return OwnedObjects && OwnedObjects->Contains(ObjectChange.Object);
+	}
+
+	void FAuthorityManager::OnClientLeft(const FClientId& ClientEndpointId)
+	{
+		ClientAuthorityData.Remove(ClientEndpointId);
+	}
+
+	EConcertSessionResponseCode FAuthorityManager::HandleChangeAuthorityRequest(
+		const FConcertSessionContext& ConcertSessionContext,
+		const FConcertChangeAuthority_Request& Request,
+		FConcertChangeAuthority_Response& Response
+		)
+	{
+		// This log does two things: 1. Identify issues in unit tests / at runtime 2. Warn about possibly malicious attempts when the server runs.
+		UE_CLOG(Request.TakeAuthority.IsEmpty() && Request.ReleaseAuthority.IsEmpty(), LogConcert, Warning, TEXT("Received invalid authority request (TakeAuthority.Num() == 0 and ReleaseAuthority.Num() == 0)"));
+		
+		FClientAuthorityData& AuthorityData = ClientAuthorityData.FindOrAdd(ConcertSessionContext.SourceEndpointId);
+		const FClientId& ClientId = ConcertSessionContext.SourceEndpointId;
+		Private::ForEachReplicatedObject(Request.TakeAuthority, [this, &Response, &AuthorityData, &ClientId](const FStreamId& StreamId, const FSoftObjectPath& ObjectPath)
+		{
+			if (CanTakeAuthority(AuthorityData, ClientId, { StreamId, ObjectPath }))
+			{
+				UE_LOG(LogConcert, Log, TEXT("Transferred authority of %s to client %s for their stream %s"), *ObjectPath.ToString(), *ClientId.ToString(EGuidFormats::Short), *StreamId.ToString(EGuidFormats::Short));
+				AuthorityData.OwnedObjects.FindOrAdd(StreamId).Add(ObjectPath);
+			}
+			else
+			{
+				UE_LOG(LogConcert, Log, TEXT("Rejected %s request of authority over %s in stream %s"), *ClientId.ToString(EGuidFormats::Short), *ObjectPath.ToString(), *StreamId.ToString(EGuidFormats::Short));
+				Response.RejectedObjects.FindOrAdd(ObjectPath).StreamIds.Add(ClientId);
+			}
+		});
+		
+		Private::ForEachReplicatedObject(Request.TakeAuthority, [this, &Response, &AuthorityData, &ClientId](const FStreamId& StreamId, const FSoftObjectPath& ObjectPath)
+		{
+			if (CanTakeAuthority(AuthorityData, ClientId, { StreamId, ObjectPath }))
+			{
+				UE_LOG(LogConcert, Log, TEXT("Transferred authority of %s to client %s for their stream %s"), *ObjectPath.ToString(), *ClientId.ToString(EGuidFormats::Short), *StreamId.ToString(EGuidFormats::Short));
+				AuthorityData.OwnedObjects.FindOrAdd(StreamId).Add(ObjectPath);
+			}
+			else
+			{
+				UE_LOG(LogConcert, Log, TEXT("Rejected %s request of authority over %s in stream %s"), *ClientId.ToString(EGuidFormats::Short), *ObjectPath.ToString(), *StreamId.ToString(EGuidFormats::Short));
+				Response.RejectedObjects.Add(ObjectPath);
+			}
+		});
+		
+		Private::ForEachReplicatedObject(Request.ReleaseAuthority, [this, &Response, &AuthorityData, &ClientId](const FStreamId& StreamId, const FSoftObjectPath& ObjectPath)
+		{
+			TSet<FSoftObjectPath>* OwnedObjects = AuthorityData.OwnedObjects.Find(StreamId);
+			// Though dubious, it is a valid request for the client to release non-owned objects
+			if (!OwnedObjects)
+			{
+				return;
+			}
+
+			OwnedObjects->Remove(ObjectPath);
+			// Avoid memory leaking
+			if (OwnedObjects->IsEmpty())
+			{
+				AuthorityData.OwnedObjects.Remove(StreamId);
+			}
+		});
+
+		return EConcertSessionResponseCode::Success;
+	}
+
+	bool FAuthorityManager::CanTakeAuthority(const FClientAuthorityData& ClientData, const FClientId& ClientId, const ConcertSyncCore::FStreamedObjectID& Object) const
+	{
+		/*
+		 * This function is unoptimized: it repeatedly iterates through every client's registered streams and properties.
+		 * However it should be fine since:
+		 * 1. Changing authority does not happen too often though,
+		 * 2. There will be a relatively small number of clients, streams, objects & properties.
+		 */
+		
+		const FReplicationStreamDescription* Description = FindClientStreamById(ClientId, Object.StreamId);
+		const FSoftObjectPath& ObjectPath = Object.Object;
+		const FReplicatedObjectInfo* PropertyInfo = Description ? Description->BaseDescription.ReplicationMap.ReplicatedObjects.Find(ObjectPath) : nullptr;
+		if (!PropertyInfo)
+		{
+			return false;
+		}
+
+		bool bFreeOfPropertyOverlaps = true;
+		const FClientId IgnoredClients[] = { ClientId };
+		ForEachClientWithPotentialConflict(ObjectPath, [this, &ClientId, &PropertyInfo, &bFreeOfPropertyOverlaps](const FClientId& WritingClientId, const FConcertPropertySelection& WrittenProperties) mutable
+		{
+			// At this point we know that WritingClientId is sending WrittenProperties - if the properties overlap, the authority request is not possible
+			bFreeOfPropertyOverlaps &= !WrittenProperties.OverlapsWith(PropertyInfo->PropertySelection);
+			if (bFreeOfPropertyOverlaps)
+			{
+				return EBreakBehavior::Continue;
+			}
+			return EBreakBehavior::Break;
+		}, IgnoredClients);
+		return bFreeOfPropertyOverlaps;
+	}
+	
+	const FReplicationStreamDescription* FAuthorityManager::FindClientStreamById(const FClientId& ClientId, const FStreamId& StreamId) const
+	{
+		const FReplicationStreamDescription* StreamDescription = nullptr;
+		Getters.ForEachStream(ClientId, [&StreamId, &StreamDescription](const FReplicationStreamDescription& Stream) mutable
+		{
+			if (Stream.BaseDescription.Identifier == StreamId)
+			{
+				StreamDescription = &Stream;
+				return EBreakBehavior::Break;
+			}
+			return EBreakBehavior::Continue;
+		});
+		return StreamDescription;
+	}
+
+	void FAuthorityManager::ForEachClientWithPotentialConflict(
+		const FSoftObjectPath& Object,
+		TFunctionRef<EBreakBehavior(const FClientId& ClientId, const FConcertPropertySelection& WrittenProperties)> Callback,
+		TArrayView<const FClientId> IgnoredClients
+		) const
+	{
+		Getters.ForEachSendingClient([this, &Object, &Callback, &IgnoredClients](const FGuid& ClientEndpointId)
+		{
+			if (IgnoredClients.Contains(ClientEndpointId))
+			{
+				return EBreakBehavior::Continue;
+			}
+
+			// If client has no authority at all, skip.
+			const FClientAuthorityData* OtherClientData = ClientAuthorityData.Find(ClientEndpointId);
+			if (!OtherClientData)
+			{
+				return EBreakBehavior::Continue;
+			}
+			
+			EBreakBehavior Result = EBreakBehavior::Continue;
+			Getters.ForEachStream(ClientEndpointId, [&Object, &Callback, &ClientEndpointId, OtherClientData, &Result](const FReplicationStreamDescription& Stream) mutable
+			{
+				// If client has not claimed authority over this object in this stream, skip
+				const FGuid& StreamId = Stream.BaseDescription.Identifier;
+				const TSet<FSoftObjectPath>* ClientControlledObjects = OtherClientData->OwnedObjects.Find(StreamId);
+				const bool bClientHasSomeAuthorityOverObject = ClientControlledObjects && ClientControlledObjects->Contains(Object);
+				if (!bClientHasSomeAuthorityOverObject)
+				{
+					return EBreakBehavior::Continue;
+				}
+
+				// This client is using the request object: report the potential conflict ...
+				const FObjectReplicationMap& ObjectReplicationMap = Stream.BaseDescription.ReplicationMap;
+				const FReplicatedObjectInfo* ReplicationObjectInfo = ObjectReplicationMap.ReplicatedObjects.Find(Object);
+				if (ReplicationObjectInfo && Callback(ClientEndpointId, ReplicationObjectInfo->PropertySelection) == EBreakBehavior::Break)
+				{
+					// ... conflict ends iteration
+					Result = EBreakBehavior::Break;
+					return EBreakBehavior::Break;
+				}
+				// ... conflict resolved
+				return EBreakBehavior::Continue;
+			});
+			return Result;
+		});
+	}
+}
diff --git a/Engine/Plugins/Developer/Concert/ConcertSync/ConcertSyncServer/Source/ConcertSyncServer/Private/Replication/AuthorityManager.h b/Engine/Plugins/Developer/Concert/ConcertSync/ConcertSyncServer/Source/ConcertSyncServer/Private/Replication/AuthorityManager.h
new file mode 100644
--- /dev/null
+++ b/Engine/Plugins/Developer/Concert/ConcertSync/ConcertSyncServer/Source/ConcertSyncServer/Private/Replication/AuthorityManager.h
@@ -0,0 +1,96 @@
+// Copyright Epic Games, Inc. All Rights Reserved.
+
+#pragma once
+
+#include "ConcertMessages.h"
+#include "Misc/EBreakBehavior.h"
+#include "Templates/Function.h"
+
+class IConcertSession;
+
+struct FConcertChangeAuthority_Response;
+struct FConcertChangeAuthority_Request;
+struct FConcertPropertySelection;
+struct FConcertSessionContext;
+struct FReplicationStreamDescription;
+
+namespace UE::ConcertSyncCore
+{
+	struct FStreamedObjectID;
+	struct FReplicatedObjectInfo;
+}
+
+namespace UE::ConcertSyncServer::Replication
+{
+	/**
+	 * Callbacks FAuthorityManager requires to function correctly.
+	 * Exists to make FAuthorityManager independent from all the other systems and allows mocking in tests.
+	 */
+	class IAuthorityManagerGetters
+	{
+	public:
+
+		/** Provides a way to extract all streams registered to a given client. */
+		virtual void ForEachStream(const FGuid& ClientEndpointId, TFunctionRef<EBreakBehavior(const FReplicationStreamDescription& Stream)> Callback) = 0;
+
+		/** Iterates through all clients have registered to send any data. */
+		virtual void ForEachSendingClient(TFunctionRef<EBreakBehavior(const FGuid& ClientEndpointId)> Callback) = 0;
+
+		virtual ~IAuthorityManagerGetters() = default;
+	};
+	
+	/** Responds to FConcertChangeAuthority_Request and tracks what objects and properties clients have authority over. */
+	class FAuthorityManager : public FNoncopyable
+	{
+	public:
+		
+		using FClientId = FGuid;
+
+		FAuthorityManager(IAuthorityManagerGetters& Getters, TSharedRef<IConcertSession> InSession);
+		~FAuthorityManager();
+
+		/**
+		 * Checks whether the client that sent the identified object had authority to send it.
+		 * @return Whether the server should process the object change.
+		 */
+		bool IsObjectChangeAllowed(const ConcertSyncCore::FReplicatedObjectInfo& ObjectChange) const;
+
+		/** Notifies this manager that the client has left, which means all their authority is now gone. */
+		void OnClientLeft(const FClientId& ClientEndpointId);
+
+	private:
+
+		/** Callbacks required to obtain client info. */
+		IAuthorityManagerGetters& Getters;
+		/** The session under which this manager operates. */
+		TSharedRef<IConcertSession> Session;
+
+		using FStreamId = FGuid;
+		struct FClientAuthorityData
+		{
+			/** Objects the client has authority over. */
+			TMap<FStreamId, TSet<FSoftObjectPath>> OwnedObjects;
+		};
+		TMap<FClientId, FClientAuthorityData> ClientAuthorityData;
+
+		EConcertSessionResponseCode HandleChangeAuthorityRequest(
+			const FConcertSessionContext& ConcertSessionContext,
+			const FConcertChangeAuthority_Request& Request,
+			FConcertChangeAuthority_Response& Response
+			);
+
+		/** Whether it is legal for this the client identified by ClientId to take control over the object given the stream the client has registered. */
+		bool CanTakeAuthority(const FClientAuthorityData& ClientData, const FClientId& ClientId, const ConcertSyncCore::FStreamedObjectID& Object) const;
+
+		/** Finds a stream registered with the client by its ID. */
+		const FReplicationStreamDescription* FindClientStreamById(const FClientId& ClientId, const FStreamId& StreamId) const;
+
+		/** Iterates through all clients that are already writing to Object. */
+		void ForEachClientWithPotentialConflict(
+			const FSoftObjectPath& Object,
+			TFunctionRef<EBreakBehavior(const FClientId& ClientId, const FConcertPropertySelection& WrittenProperties)> Callback,
+			TArrayView<const FClientId> IgnoredClients = {}
+			) const;
+	};
+}
+
diff --git a/Engine/Plugins/Developer/Concert/ConcertSync/ConcertSyncServer/Source/ConcertSyncServer/Private/Replication/ConcertServerReplicationManager.cpp b/Engine/Plugins/Developer/Concert/ConcertSync/ConcertSyncServer/Source/ConcertSyncServer/Private/Replication/ConcertServerReplicationManager.cpp
--- a/Engine/Plugins/Developer/Concert/ConcertSync/ConcertSyncServer/Source/ConcertSyncServer/Private/Replication/ConcertServerReplicationManager.cpp
+++ b/Engine/Plugins/Developer/Concert/ConcertSync/ConcertSyncServer/Source/ConcertSyncServer/Private/Replication/ConcertServerReplicationManager.cpp
@@ -2,6 +2,7 @@
 
 #include "ConcertServerReplicationManager.h"
 
+#include "AuthorityManager.h"
 #include "ConcertLogGlobal.h"
 #include "IConcertSession.h"
 #include "Replication/ConcertReplicationClient.h"
@@ -15,8 +16,9 @@ namespace UE::ConcertSyncServer::Replication
 	FConcertServerReplicationManager::FConcertServerReplicationManager(TSharedRef<IConcertServerSession> InLiveSession)
 		: Session(MoveTemp(InLiveSession))
 		, ReplicationFormat(MakeShared<ConcertSyncCore::FFullObjectFormat>())
+		, AuthorityManager(MakeShared<FAuthorityManager>(*this, Session))
 		, ReplicationCache(MakeShared<ConcertSyncCore::FObjectReplicationCache>(ReplicationFormat))
-		, ReplicationDataReceiver(Session, ReplicationCache)
+		, ReplicationDataReceiver(AuthorityManager, Session, ReplicationCache)
 	{
 		Session->RegisterCustomRequestHandler<FConcertReplication_Join_Request, FConcertReplication_Join_Response>(this, &FConcertServerReplicationManager::HandleJoinReplicationSessionRequest);
 		Session->RegisterCustomEventHandler<FConcertReplication_LeaveEvent>(this, &FConcertServerReplicationManager::HandleLeaveReplicationSessionRequest);
@@ -33,6 +35,35 @@ namespace UE::ConcertSyncServer::Replication
 		Session->OnTick().RemoveAll(this);
 	}
 
+	void FConcertServerReplicationManager::ForEachStream(const FGuid& ClientEndpointId, TFunctionRef<EBreakBehavior(const FReplicationStreamDescription& Stream)> Callback)
+	{
+		const TSharedRef<FConcertReplicationClient>* Client = Clients.Find(ClientEndpointId);
+		if (!ensure(Client))
+		{
+			return;
+		}
+
+		for (const FReplicationStreamDescription& Stream : (*Client)->GetStreamDescriptions())
+		{
+			if (Callback(Stream) == EBreakBehavior::Break)
+			{
+				break;
+			}
+		}
+	}
+
+	void FConcertServerReplicationManager::ForEachSendingClient(TFunctionRef<EBreakBehavior(const FGuid& ClientEndpointId)> Callback)
+	{
+		for (const TPair<FGuid, TSharedRef<FConcertReplicationClient>>& ClientPair : Clients)
+		{
+			if (!ClientPair.Value->GetStreamDescriptions().IsEmpty()
+				&& Callback(ClientPair.Key) == EBreakBehavior::Break)
+			{
+				break;
+			}
+		}
+	}
+
 	EConcertSessionResponseCode FConcertServerReplicationManager::HandleJoinReplicationSessionRequest(
 		const FConcertSessionContext& ConcertSessionContext,
 		const FConcertReplication_Join_Request& Request,
@@ -69,16 +100,7 @@ namespace UE::ConcertSyncServer::Replication
 			return EConcertSessionResponseCode::Success;
 		}
 
-		auto GetClientName = [this](const FGuid& ClientEndpointId) -> FString
-		{
-			FConcertSessionClientInfo Info;
-			if (!ensure(Session->FindSessionClient(ClientEndpointId, Info)))
-			{
-				return TEXT("Unknown Client");
-			}
-			return FString::Printf(TEXT("%s (%s)"), *Info.ClientInfo.UserName, *ClientEndpointId.ToString());
-		};
-		auto[ErrorCode, ErrorMessage, StreamDescriptions] = ValidateRequest(Request, Clients, GetClientName);
+		auto[ErrorCode, ErrorMessage, StreamDescriptions] = ValidateRequest(Request);
 		if (ErrorCode != EJoinReplicationErrorCode::Success)
 		{
 			Response = { ErrorCode, ErrorMessage };
@@ -95,15 +117,20 @@ namespace UE::ConcertSyncServer::Replication
 		const FConcertReplication_LeaveEvent& EventData
 		)
 	{
-		UE_LOG(LogConcert, Log, TEXT("Received replication leave request from endpoint %s"), *ConcertSessionContext.SourceEndpointId.ToString());
-		Clients.Remove(ConcertSessionContext.SourceEndpointId);
+		const FGuid ClientEndpointId = ConcertSessionContext.SourceEndpointId;
+		UE_LOG(LogConcert, Log, TEXT("Received replication leave request from endpoint %s"), *ClientEndpointId.ToString());
+		
+		Clients.Remove(ClientEndpointId);
+		AuthorityManager->OnClientLeft(ClientEndpointId);
 	}
 
 	void FConcertServerReplicationManager::OnConnectionChanged(IConcertServerSession& ConcertServerSession, EConcertClientStatus ConcertClientStatus, const FConcertSessionClientInfo& ClientInfo)
 	{
+		const FGuid ClientEndpointId = ClientInfo.ClientEndpointId;
 		if (ConcertClientStatus == EConcertClientStatus::Disconnected)
 		{
-			Clients.Remove(ClientInfo.ClientEndpointId);
+			Clients.Remove(ClientEndpointId);
+			AuthorityManager->OnClientLeft(ClientEndpointId);
 		}
 	}
 
diff --git a/Engine/Plugins/Developer/Concert/ConcertSync/ConcertSyncServer/Source/ConcertSyncServer/Private/Replication/ConcertServerReplicationManager.h b/Engine/Plugins/Developer/Concert/ConcertSync/ConcertSyncServer/Source/ConcertSyncServer/Private/Replication/ConcertServerReplicationManager.h
--- a/Engine/Plugins/Developer/Concert/ConcertSync/ConcertSyncServer/Source/ConcertSyncServer/Private/Replication/ConcertServerReplicationManager.h
+++ b/Engine/Plugins/Developer/Concert/ConcertSync/ConcertSyncServer/Source/ConcertSyncServer/Private/Replication/ConcertServerReplicationManager.h
@@ -2,14 +2,16 @@
 
 #pragma once
 
+#include "AuthorityManager.h"
 #include "ConcertMessages.h"
 #include "ConcertReplicationClient.h"
-#include "IConcertSessionHandler.h"
+#include "Processing/ServerObjectReplicationReceiver.h"
 #include "Replication/IConcertServerReplicationManager.h"
 #include "Replication/Messages/ConcertReplicationHandshakeMessages.h"
-#include "Replication/Processing/ObjectReplicationReceiver.h"
+
 #include "Templates/SharedPointer.h"
 #include "Templates/Tuple.h"
+#include "Templates/UnrealTemplate.h"
 
 namespace UE::ConcertSyncCore
 {
@@ -22,12 +24,29 @@ class IConcertServerSession;
 
 namespace UE::ConcertSyncServer::Replication
 {
-	class FConcertServerReplicationManager : public IConcertServerReplicationManager
+	class FAuthorityManager;
+
+	/**
+	 * Manages all server-side systems relevant to the Replication features.
+	 * 
+	 * Primarily responds to client requests to join (handshake) and leave replication delegating the result of the
+	 * operation to relevant systems. 
+	 */
+	class FConcertServerReplicationManager
+		: public IConcertServerReplicationManager
+		, public IAuthorityManagerGetters
+		, public FNoncopyable
 	{
 	public:
+		
 		explicit FConcertServerReplicationManager(TSharedRef<IConcertServerSession> InLiveSession);
 		virtual ~FConcertServerReplicationManager() override;
 
+		//~ Begin IAuthorityManagerGetters Interface
+		virtual void ForEachStream(const FGuid& ClientEndpointId, TFunctionRef<EBreakBehavior(const FReplicationStreamDescription& Stream)> Callback) override;
+		virtual void ForEachSendingClient(TFunctionRef<EBreakBehavior(const FGuid& ClientEndpointId)> Callback) override;
+		//~ End IAuthorityManagerGetters Interface
+
 	private:
 		
 		/** Session instance this manager was created for. */
@@ -35,11 +54,14 @@ namespace UE::ConcertSyncServer::Replication
 		
 		/** Responsible for analysing received replication data. */
 		TSharedRef<ConcertSyncCore::IObjectReplicationFormat> ReplicationFormat;
+
+		/** Responds to client requests to changing authority and can be asked whether an object change is valid to take place. */
+		TSharedRef<FAuthorityManager> AuthorityManager;
 		
 		/** Received replication events are put into the ReplicationCache. The cache is used to relay data to clients latently. */
 		TSharedRef<ConcertSyncCore::FObjectReplicationCache> ReplicationCache;
 		/** Receives replication events from all endpoints. */
-		ConcertSyncCore::FObjectReplicationReceiver ReplicationDataReceiver;
+		FServerObjectReplicationReceiver ReplicationDataReceiver;
 
 		/** Clients that have requested to join replication. Maps client ID to replication info. */
 		TMap<FGuid, TSharedRef<FConcertReplicationClient>> Clients;
@@ -51,6 +73,7 @@ namespace UE::ConcertSyncServer::Replication
 		
 		void OnConnectionChanged(IConcertServerSession& ConcertServerSession, EConcertClientStatus ConcertClientStatus, const FConcertSessionClientInfo& ConcertSessionClientInfo);
 
+	private:
 		/**
 		 * Ticks all clients which causes clients to process pending data and send it to the corresponding endpoints.
 		 * 
diff --git a/Engine/Plugins/Developer/Concert/ConcertSync/ConcertSyncServer/Source/ConcertSyncServer/Private/Replication/Processing/ServerObjectReplicationReceiver.cpp b/Engine/Plugins/Developer/Concert/ConcertSync/ConcertSyncServer/Source/ConcertSyncServer/Private/Replication/Processing/ServerObjectReplicationReceiver.cpp
new file mode 100644
--- /dev/null
+++ b/Engine/Plugins/Developer/Concert/ConcertSync/ConcertSyncServer/Source/ConcertSyncServer/Private/Replication/Processing/ServerObjectReplicationReceiver.cpp
@@ -0,0 +1,30 @@
+// Copyright Epic Games, Inc. All Rights Reserved.
+
+#include "ServerObjectReplicationReceiver.h"
+
+#include "IConcertSessionHandler.h"
+#include "Replication/AuthorityManager.h"
+#include "Replication/ReplicationStreamObjectID.h"
+#include "Replication/Messages/ConcertReplicationEvents.h"
+
+namespace UE::ConcertSyncServer::Replication
+{
+	FServerObjectReplicationReceiver::FServerObjectReplicationReceiver(
+		TSharedRef<FAuthorityManager> AuthorityManager,
+		TSharedRef<IConcertSession> Session,
+		TSharedRef<ConcertSyncCore::FObjectReplicationCache> ReplicationCache
+		)
+		: FObjectReplicationReceiver(MoveTemp(Session), MoveTemp(ReplicationCache))
+		, AuthorityManager(MoveTemp(AuthorityManager))
+	{}
+
+	bool FServerObjectReplicationReceiver::ShouldAcceptObject(
+		const FConcertSessionContext& SessionContext,
+		const FConcertStreamReplicationEvent& StreamEvent,
+		const FConcertObjectReplicationEvent& ObjectEvent
+		) const
+	{
+		const ConcertSyncCore::FReplicatedObjectInfo ReplicatedObjectInfo { { StreamEvent.StreamId, ObjectEvent.ReplicatedObject }, SessionContext.SourceEndpointId };
+		return AuthorityManager->IsObjectChangeAllowed(ReplicatedObjectInfo);
+	}
+}
diff --git a/Engine/Plugins/Developer/Concert/ConcertSync/ConcertSyncServer/Source/ConcertSyncServer/Private/Replication/Processing/ServerObjectReplicationReceiver.h b/Engine/Plugins/Developer/Concert/ConcertSync/ConcertSyncServer/Source/ConcertSyncServer/Private/Replication/Processing/ServerObjectReplicationReceiver.h
new file mode 100644
--- /dev/null
+++ b/Engine/Plugins/Developer/Concert/ConcertSync/ConcertSyncServer/Source/ConcertSyncServer/Private/Replication/Processing/ServerObjectReplicationReceiver.h
@@ -0,0 +1,32 @@
+// Copyright Epic Games, Inc. All Rights Reserved.
+
+#pragma once
+
+#include "Replication/Processing/ObjectReplicationReceiver.h"
+#include "Templates/SharedPointer.h"
+
+namespace UE::ConcertSyncServer::Replication
+{
+	class FAuthorityManager;
+	
+	/** Rejects changes to objects that the sending client does not have authority over. */
+	class FServerObjectReplicationReceiver : public ConcertSyncCore::FObjectReplicationReceiver
+	{
+	public:
+
+		FServerObjectReplicationReceiver(TSharedRef<FAuthorityManager> AuthorityManager, TSharedRef<IConcertSession> Session, TSharedRef<ConcertSyncCore::FObjectReplicationCache> ReplicationCache);
+
+	protected:
+
+		//~ Begin FObjectReplicationReceiver Interface
+		virtual bool ShouldAcceptObject(const FConcertSessionContext& SessionContext, const FConcertStreamReplicationEvent& StreamEvent, const FConcertObjectReplicationEvent& ObjectEvent) const override;
+		//~ End FObjectReplicationReceiver Interface
+
+	private:
+
+		/** Used to determine whether a client has authority over objects. */
+		TSharedRef<FAuthorityManager> AuthorityManager;
+	};
+}
+
+
diff --git a/Engine/Plugins/Developer/Concert/ConcertSync/ConcertSyncServer/Source/ConcertSyncServer/Private/Replication/Util/JoinRequestValidation.cpp b/Engine/Plugins/Developer/Concert/ConcertSync/ConcertSyncServer/Source/ConcertSyncServer/Private/Replication/Util/JoinRequestValidation.cpp
--- a/Engine/Plugins/Developer/Concert/ConcertSync/ConcertSyncServer/Source/ConcertSyncServer/Private/Replication/Util/JoinRequestValidation.cpp
+++ b/Engine/Plugins/Developer/Concert/ConcertSync/ConcertSyncServer/Source/ConcertSyncServer/Private/Replication/Util/JoinRequestValidation.cpp
@@ -9,57 +9,24 @@
 
 namespace UE::ConcertSyncServer::Replication
 {
-	static TTuple<EJoinReplicationErrorCode, FString> CheckInputConstraints(
-		const FConcertReplication_Join_Request& Request,
-		const TMap<FGuid, TSharedRef<FConcertReplicationClient>>& Clients,
-		FGetClientName GetClientNameFunc
-		)
+	static TTuple<EJoinReplicationErrorCode, FString> CheckInputConstraints(const FConcertReplication_Join_Request& Request)
 	{
-		/* Duplicate objects and properties might not actually cause any trouble.
+		/* Duplicate objects and properties in a request might not actually cause any trouble.
 		 * However, we will straight out reject any input that is not exactly how we expect it.
 		 * This serves us to
-		 *		1. reason about our pre/post conditions
-		 *		2. since other code can assume said conditions, it is a common security mechanic to reject input in case somebody sends us specifically crafted input
+		 *	1. reason about our pre/post conditions
+		 *	2. since other code can assume said conditions, it is a common security mechanic to reject input in case somebody sends us specifically crafted input
 		 */
-		TMap<FSoftObjectPath, TArray<const FReplicatedObjectInfo*>> UniqueObjectPropertyDetection;
-		TMap<FGuid, FConcertReplicationClient*> UniqueStreamIdDetection;
-		TSet<FGuid> UniqueStreamsFromRequestor;
-
-		// Prepare properties and stream IDs already used by other clients 
-		for (const TPair<FGuid, TSharedRef<FConcertReplicationClient>>& Client : Clients)
-		{
-			for (const FReplicationStreamDescription& StreamDescription : Client.Value->GetStreamDescriptions())
-			{
-				for (const TPair<FSoftObjectPath, FReplicatedObjectInfo>& ObjectSelection : StreamDescription.BaseDescription.ReplicationMap.ReplicatedObjects)
-				{
-					UniqueObjectPropertyDetection.FindOrAdd(ObjectSelection.Key).Add(&ObjectSelection.Value);
-				}
-				
-				const FGuid& ExistingStreamId = StreamDescription.BaseDescription.Identifier;
-				check(!UniqueStreamIdDetection.Contains(ExistingStreamId))
-				UniqueStreamIdDetection.Add(ExistingStreamId, &Client.Value.Get());
-			}
-		}
 		
+		TSet<FGuid> UniqueStreamsFromRequestor;
 		for (const FReplicationStreamDescription_NetPacked& Stream : Request.Streams)
 		{
 			const FGuid& StreamId = Stream.BaseDescription.Identifier;
 
-			// Another client already using this stream?
-			if (FConcertReplicationClient** OwningClient = UniqueStreamIdDetection.Find(StreamId))
-			{
-				const FString ErrorMessage = FString::Printf(
-					TEXT("Stream %s is already in use by client %s"),
-					*StreamId.ToString(),
-					*GetClientNameFunc((*OwningClient)->GetClientEndpointId())
-					);
-				return { EJoinReplicationErrorCode::ConflictingStreamId, ErrorMessage };
-			}
-
 			// Request contained stream ID twice?
 			if (UniqueStreamsFromRequestor.Contains(StreamId))
 			{
-				return { EJoinReplicationErrorCode::ConflictingStreamId,  FString::Printf(TEXT("Duplicate stream ID %s in request."), *StreamId.ToString()) };
+				return { EJoinReplicationErrorCode::DuplicateStreamId,  FString::Printf(TEXT("Duplicate stream ID %s in request."), *StreamId.ToString()) };
 			}
 			UniqueStreamsFromRequestor.Add(Stream.BaseDescription.Identifier);
 
@@ -82,7 +49,6 @@ namespace UE::ConcertSyncServer::Replication
 
 				// Validate unique properties
 				TSet<FConcertPropertyChain> UniquePropertyDetection;
-				const TArray<const FReplicatedObjectInfo*>* PreexistingPropertySelections = UniqueObjectPropertyDetection.Find(ObjectPath);
 				for (const FConcertPropertyChain& PropertyChain : ObjectInfo.PropertySelection.ReplicatedProperties)
 				{
 					// No duplicate properties!
@@ -97,58 +63,16 @@ namespace UE::ConcertSyncServer::Replication
 						return { EJoinReplicationErrorCode::DuplicateProperty, ErrorMessage };
 					}
 					UniquePropertyDetection.Add(PropertyChain);
-
-					if (!PreexistingPropertySelections)
-					{
-						continue;
-					}
-
-					// Is another client already streaming these properties?
-					for (const FReplicatedObjectInfo* PreexistingObjectSelection : *PreexistingPropertySelections)
-					{
-						// The request's object class should be the same the class previously assigned to the object in previous requests
-						if (PreexistingObjectSelection->ClassPath != ObjectInfo.ClassPath)
-						{
-							const FString ErrorMessage = FString::Printf(
-								TEXT("Class %s in stream %s for object %s does not match expected class %s (another client has told us this object has this class)."),
-								*ObjectInfo.ClassPath.ToString(),
-								*Stream.BaseDescription.Identifier.ToString(),
-								*ObjectPath.ToString(),
-								*PreexistingObjectSelection->ClassPath.ToString()
-							);
-							return { EJoinReplicationErrorCode::InvalidClass, ErrorMessage};
-						}
-
-						// Another client (or config in this request) is already using this property
-						if (PreexistingObjectSelection->PropertySelection.ReplicatedProperties.Contains(PropertyChain))
-						{
-							// TODO: In the future we'll change the handshake to no longer reject when another client wants to replicate this property. Adjust the message then.
-							// (Instead every client will request authority for each object before replicating)
-							const FString ErrorMessage = FString::Printf(
-								TEXT("Property %s in stream %s for object %s is already in duplicate or in use by another client"),
-								*PropertyChain.ToString(),
-								*Stream.BaseDescription.Identifier.ToString(),
-								*ObjectPath.ToString()
-							);
-							return { EJoinReplicationErrorCode::ConflictingAuthority, ErrorMessage };
-						}
-					}
 				}
-				
-				UniqueObjectPropertyDetection.FindOrAdd(ObjectPath).Add(&ObjectInfo);
 			}
 		}
 
 		return { EJoinReplicationErrorCode::Success, TEXT("") };
 	}
 	
-	TTuple<EJoinReplicationErrorCode, FString, TArray<FReplicationStreamDescription>> ValidateRequest(
-		const FConcertReplication_Join_Request& Request,
-		const TMap<FGuid, TSharedRef<FConcertReplicationClient>>& Clients,
-		FGetClientName GetClientNameFunc
-		)
+	TTuple<EJoinReplicationErrorCode, FString, TArray<FReplicationStreamDescription>> ValidateRequest(const FConcertReplication_Join_Request& Request)
 	{
-		if (auto[ErrorCode, ErrorMessage] = CheckInputConstraints(Request, Clients, GetClientNameFunc)
+		if (auto[ErrorCode, ErrorMessage] = CheckInputConstraints(Request)
 			; ErrorCode != EJoinReplicationErrorCode::Success)
 		{
 			return { ErrorCode, ErrorMessage, TArray<FReplicationStreamDescription>{} };
diff --git a/Engine/Plugins/Developer/Concert/ConcertSync/ConcertSyncServer/Source/ConcertSyncServer/Private/Replication/Util/JoinRequestValidation.h b/Engine/Plugins/Developer/Concert/ConcertSync/ConcertSyncServer/Source/ConcertSyncServer/Private/Replication/Util/JoinRequestValidation.h
--- a/Engine/Plugins/Developer/Concert/ConcertSync/ConcertSyncServer/Source/ConcertSyncServer/Private/Replication/Util/JoinRequestValidation.h
+++ b/Engine/Plugins/Developer/Concert/ConcertSync/ConcertSyncServer/Source/ConcertSyncServer/Private/Replication/Util/JoinRequestValidation.h
@@ -29,9 +29,5 @@ namespace UE::ConcertSyncServer::Replication
 	 * 
 	 * @return The error code for the request and the unpacked stream descriptions
 	 */
-	TTuple<EJoinReplicationErrorCode, FString, TArray<FReplicationStreamDescription>> ValidateRequest(
-		const FConcertReplication_Join_Request& Request,
-		const TMap<FGuid, TSharedRef<FConcertReplicationClient>>& Clients,
-		FGetClientName GetClientNameFunc
-		);
+	TTuple<EJoinReplicationErrorCode, FString, TArray<FReplicationStreamDescription>> ValidateRequest(const FConcertReplication_Join_Request& Request);
 }
diff --git a/Engine/Plugins/Developer/Concert/ConcertSync/ConcertSyncTest/Source/ConcertSyncTest/Private/Replication/ConcertPropertyTests.cpp b/Engine/Plugins/Developer/Concert/ConcertSync/ConcertSyncTest/Source/ConcertSyncTest/Private/Replication/ConcertPropertyTests.cpp
--- a/Engine/Plugins/Developer/Concert/ConcertSync/ConcertSyncTest/Source/ConcertSyncTest/Private/Replication/ConcertPropertyTests.cpp
+++ b/Engine/Plugins/Developer/Concert/ConcertSync/ConcertSyncTest/Source/ConcertSyncTest/Private/Replication/ConcertPropertyTests.cpp
@@ -5,466 +5,514 @@
 #include "Replication/ReplicationPropertyFilter.h"
 #include "TestReflectionObject.h"
 
+#include "Containers/UnrealString.h"
 #include "Misc/AutomationTest.h"
 #include "Serialization/ArchiveSerializedPropertyChain.h"
 #include "Serialization/ObjectAndNameAsStringProxyArchive.h"
 #include "Serialization/ObjectReader.h"
 #include "Serialization/ObjectWriter.h"
 
-/**
- * Tests that the FConcertPropertyChain constructor works as the documentation dictates:
- *  - Root properties
- *  - Array, set, map of
- *		- "normal" structs, like FVector > lists sub-properties
- *		- primitives and native structs > Path contains FConcertPropertyChain::InternalContainerPropertyValueName ("Value") at the end
- */
-IMPLEMENT_SIMPLE_AUTOMATION_TEST(FConcertPropertyChainTests, "Concert.Replication.Data.PropertyChainTests", EAutomationTestFlags::EditorContext | EAutomationTestFlags::EngineFilter)
-bool FConcertPropertyChainTests::RunTest(const FString& Parameters)
+namespace UE::ConcertSyncTest::Replication::PropertyChain
 {
-	// The test is intentionally set up to make it it easy to set breakpoints albeit at the expense of increasing the amount of code
+	/**
+	 * Tests that the FConcertPropertyChain constructor works as the documentation dictates:
+	 *  - Root properties
+	 *  - Array, set, map of
+	 *		- "normal" structs, like FVector > lists sub-properties
+	 *		- primitives and native structs > Path contains FConcertPropertyChain::InternalContainerPropertyValueName ("Value") at the end
+	 */
+	IMPLEMENT_SIMPLE_AUTOMATION_TEST(FConcertPropertyChainTests, "Concert.Replication.Data.PropertyChainTests", EAutomationTestFlags::EditorContext | EAutomationTestFlags::EngineFilter)
+	bool FConcertPropertyChainTests::RunTest(const FString& Parameters)
+	{
+		// The test is intentionally set up to make it it easy to set breakpoints albeit at the expense of increasing the amount of code
 
-	// 1. Input
-	UStruct* TestReplicationStruct = FTestReplicationStruct::StaticStruct();
-	FProperty* ValueProp = TestReplicationStruct->FindPropertyByName(GET_MEMBER_NAME_CHECKED(FTestReplicationStruct, Value));
-	FProperty* VectorProp = TestReplicationStruct->FindPropertyByName(GET_MEMBER_NAME_CHECKED(FTestReplicationStruct, Vector));
-	FProperty* VectorXProp = CastField<FStructProperty>(VectorProp)->Struct->FindPropertyByName(GET_MEMBER_NAME_CHECKED(FVector, X));
+		// 1. Input
+		UStruct* TestReplicationStruct = FTestReplicationStruct::StaticStruct();
+		FProperty* ValueProp = TestReplicationStruct->FindPropertyByName(GET_MEMBER_NAME_CHECKED(FTestReplicationStruct, Value));
+		FProperty* VectorProp = TestReplicationStruct->FindPropertyByName(GET_MEMBER_NAME_CHECKED(FTestReplicationStruct, Vector));
+		FProperty* VectorXProp = CastField<FStructProperty>(VectorProp)->Struct->FindPropertyByName(GET_MEMBER_NAME_CHECKED(FVector, X));
 	
-	FProperty* FloatArrayProp = TestReplicationStruct->FindPropertyByName(GET_MEMBER_NAME_CHECKED(FTestReplicationStruct, FloatArray));
-	FProperty* FloatArrayInnerProp = CastField<FArrayProperty>(FloatArrayProp)->Inner;
-	FProperty* StringArrayProp = TestReplicationStruct->FindPropertyByName(GET_MEMBER_NAME_CHECKED(FTestReplicationStruct, StringArray));
-	FProperty* NativeStructArrayProp = TestReplicationStruct->FindPropertyByName(GET_MEMBER_NAME_CHECKED(FTestReplicationStruct, NativeStructArray));
-	FProperty* NativeStructArrayInnerProp = CastField<FArrayProperty>(NativeStructArrayProp)->Inner;
+		FProperty* FloatArrayProp = TestReplicationStruct->FindPropertyByName(GET_MEMBER_NAME_CHECKED(FTestReplicationStruct, FloatArray));
+		FProperty* FloatArrayInnerProp = CastField<FArrayProperty>(FloatArrayProp)->Inner;
+		FProperty* StringArrayProp = TestReplicationStruct->FindPropertyByName(GET_MEMBER_NAME_CHECKED(FTestReplicationStruct, StringArray));
+		FProperty* NativeStructArrayProp = TestReplicationStruct->FindPropertyByName(GET_MEMBER_NAME_CHECKED(FTestReplicationStruct, NativeStructArray));
+		FProperty* NativeStructArrayInnerProp = CastField<FArrayProperty>(NativeStructArrayProp)->Inner;
 	
-	FProperty* FloatSetProp = TestReplicationStruct->FindPropertyByName(GET_MEMBER_NAME_CHECKED(FTestReplicationStruct, FloatSet));
-	FProperty* FloatSetInnerProp = CastField<FSetProperty>(FloatSetProp)->ElementProp;
-	FProperty* StringSetProp = TestReplicationStruct->FindPropertyByName(GET_MEMBER_NAME_CHECKED(FTestReplicationStruct, StringSet));
-	FProperty* NativeStructSetProp = TestReplicationStruct->FindPropertyByName(GET_MEMBER_NAME_CHECKED(FTestReplicationStruct, NativeStructSet));
-	FProperty* NativeStructSetInnerProp = CastField<FSetProperty>(NativeStructSetProp)->ElementProp;
+		FProperty* FloatSetProp = TestReplicationStruct->FindPropertyByName(GET_MEMBER_NAME_CHECKED(FTestReplicationStruct, FloatSet));
+		FProperty* FloatSetInnerProp = CastField<FSetProperty>(FloatSetProp)->ElementProp;
+		FProperty* StringSetProp = TestReplicationStruct->FindPropertyByName(GET_MEMBER_NAME_CHECKED(FTestReplicationStruct, StringSet));
+		FProperty* NativeStructSetProp = TestReplicationStruct->FindPropertyByName(GET_MEMBER_NAME_CHECKED(FTestReplicationStruct, NativeStructSet));
+		FProperty* NativeStructSetInnerProp = CastField<FSetProperty>(NativeStructSetProp)->ElementProp;
 	
-	FProperty* StringToFloatProp = TestReplicationStruct->FindPropertyByName(GET_MEMBER_NAME_CHECKED(FTestReplicationStruct, StringToFloat));
-	FProperty* StringToFloatInnerProp = CastField<FMapProperty>(StringToFloatProp)->ValueProp;
-	FProperty* StringToVectorProp = TestReplicationStruct->FindPropertyByName(GET_MEMBER_NAME_CHECKED(FTestReplicationStruct, StringToVector));
-	FProperty* StringToNativeStructProp = TestReplicationStruct->FindPropertyByName(GET_MEMBER_NAME_CHECKED(FTestReplicationStruct, StringToNativeStruct));
-	FProperty* StringToNativeStructInnerProp = CastField<FMapProperty>(StringToNativeStructProp)->ValueProp;
-
-	FArchiveSerializedPropertyChain VectorXChain;
-	FArchiveSerializedPropertyChain FloatArrayChain;
-	FArchiveSerializedPropertyChain FloatSetArrayChain;
-	FArchiveSerializedPropertyChain StringToFloatChain;
-	FArchiveSerializedPropertyChain NativeStructArrayChain;
-	FArchiveSerializedPropertyChain NativeStructSetChain;
-	FArchiveSerializedPropertyChain StringToNativeStructChain;
-	VectorXChain.PushProperty(VectorProp, false);
-	FloatArrayChain.PushProperty(FloatArrayProp, false);
-	FloatSetArrayChain.PushProperty(FloatSetProp, false);
-	StringToFloatChain.PushProperty(StringToFloatProp, false);
-	NativeStructArrayChain.PushProperty(NativeStructArrayProp, false);
-	NativeStructSetChain.PushProperty(NativeStructSetProp, false);
-	StringToNativeStructChain.PushProperty(StringToNativeStructProp, false);
+		FProperty* StringToFloatProp = TestReplicationStruct->FindPropertyByName(GET_MEMBER_NAME_CHECKED(FTestReplicationStruct, StringToFloat));
+		FProperty* StringToFloatInnerProp = CastField<FMapProperty>(StringToFloatProp)->ValueProp;
+		FProperty* StringToVectorProp = TestReplicationStruct->FindPropertyByName(GET_MEMBER_NAME_CHECKED(FTestReplicationStruct, StringToVector));
+		FProperty* StringToNativeStructProp = TestReplicationStruct->FindPropertyByName(GET_MEMBER_NAME_CHECKED(FTestReplicationStruct, StringToNativeStruct));
+		FProperty* StringToNativeStructInnerProp = CastField<FMapProperty>(StringToNativeStructProp)->ValueProp;
+
+		FArchiveSerializedPropertyChain VectorXChain;
+		FArchiveSerializedPropertyChain FloatArrayChain;
+		FArchiveSerializedPropertyChain FloatSetArrayChain;
+		FArchiveSerializedPropertyChain StringToFloatChain;
+		FArchiveSerializedPropertyChain NativeStructArrayChain;
+		FArchiveSerializedPropertyChain NativeStructSetChain;
+		FArchiveSerializedPropertyChain StringToNativeStructChain;
+		VectorXChain.PushProperty(VectorProp, false);
+		FloatArrayChain.PushProperty(FloatArrayProp, false);
+		FloatSetArrayChain.PushProperty(FloatSetProp, false);
+		StringToFloatChain.PushProperty(StringToFloatProp, false);
+		NativeStructArrayChain.PushProperty(NativeStructArrayProp, false);
+		NativeStructSetChain.PushProperty(NativeStructSetProp, false);
+		StringToNativeStructChain.PushProperty(StringToNativeStructProp, false);
 	
 
 	
-	// 2. Run
-	const FConcertPropertyChain Value(nullptr, *ValueProp);
-	const FConcertPropertyChain Vector(nullptr, *VectorProp);
-	const FConcertPropertyChain VectorX(&VectorXChain, *VectorXProp);
+		// 2. Run
+		const FConcertPropertyChain Value(nullptr, *ValueProp);
+		const FConcertPropertyChain Vector(nullptr, *VectorProp);
+		const FConcertPropertyChain VectorX(&VectorXChain, *VectorXProp);
 	
-	const FConcertPropertyChain FloatArray(nullptr, *FloatArrayProp);
-	const FConcertPropertyChain FloatArrayInner(&FloatArrayChain, *FloatArrayInnerProp);
-	const FConcertPropertyChain StringArray(nullptr, *StringArrayProp);
-	const FConcertPropertyChain NativeStructArray(nullptr, *NativeStructArrayProp);
-	const FConcertPropertyChain NativeStructArrayInner(&NativeStructArrayChain, *NativeStructArrayInnerProp);
+		const FConcertPropertyChain FloatArray(nullptr, *FloatArrayProp);
+		const FConcertPropertyChain FloatArrayInner(&FloatArrayChain, *FloatArrayInnerProp);
+		const FConcertPropertyChain StringArray(nullptr, *StringArrayProp);
+		const FConcertPropertyChain NativeStructArray(nullptr, *NativeStructArrayProp);
+		const FConcertPropertyChain NativeStructArrayInner(&NativeStructArrayChain, *NativeStructArrayInnerProp);
 	
-	const FConcertPropertyChain FloatSet(nullptr, *FloatSetProp);
-	const FConcertPropertyChain FloatSetInner(&FloatSetArrayChain, *FloatSetInnerProp);
-	const FConcertPropertyChain StringSet(nullptr, *StringSetProp);
-	const FConcertPropertyChain NativeStructSet(nullptr, *NativeStructSetProp);
-	const FConcertPropertyChain NativeStructSetInner(&NativeStructSetChain, *NativeStructSetInnerProp);
+		const FConcertPropertyChain FloatSet(nullptr, *FloatSetProp);
+		const FConcertPropertyChain FloatSetInner(&FloatSetArrayChain, *FloatSetInnerProp);
+		const FConcertPropertyChain StringSet(nullptr, *StringSetProp);
+		const FConcertPropertyChain NativeStructSet(nullptr, *NativeStructSetProp);
+		const FConcertPropertyChain NativeStructSetInner(&NativeStructSetChain, *NativeStructSetInnerProp);
 	
-	const FConcertPropertyChain StringToFloat(nullptr, *StringToFloatProp);
-	const FConcertPropertyChain StringToFloatInner(&StringToFloatChain, *StringToFloatInnerProp);
-	const FConcertPropertyChain StringToVector(nullptr, *StringToVectorProp);
-	const FConcertPropertyChain StringToNativeStruct(nullptr, *StringToNativeStructProp);
-	const FConcertPropertyChain StringToNativeStructInner(&StringToNativeStructChain, *StringToNativeStructInnerProp);
+		const FConcertPropertyChain StringToFloat(nullptr, *StringToFloatProp);
+		const FConcertPropertyChain StringToFloatInner(&StringToFloatChain, *StringToFloatInnerProp);
+		const FConcertPropertyChain StringToVector(nullptr, *StringToVectorProp);
+		const FConcertPropertyChain StringToNativeStruct(nullptr, *StringToNativeStructProp);
+		const FConcertPropertyChain StringToNativeStructInner(&StringToNativeStructChain, *StringToNativeStructInnerProp);
 
 
 	
-	// 3. Test
-	TestTrue(TEXT("Value"), Value == TArray<FName>{ TEXT("Value") });
-	TestTrue(TEXT("Vector"), Vector == TArray<FName>{ TEXT("Vector") });
-	TestTrue(TEXT("Vector.X"), VectorX == TArray<FName>{ FName(TEXT("Vector")), FName(TEXT("X")) });
+		// 3. Test
+		TestTrue(TEXT("Value"), Value == TArray<FName>{ TEXT("Value") });
+		TestTrue(TEXT("Vector"), Vector == TArray<FName>{ TEXT("Vector") });
+		TestTrue(TEXT("Vector.X"), VectorX == TArray<FName>{ FName(TEXT("Vector")), FName(TEXT("X")) });
 	
-	TestTrue(TEXT("FloatArray"), FloatArray == TArray<FName>{ TEXT("FloatArray") });
-	TestTrue(TEXT("FloatArrayInner"), FloatArrayInner == TArray<FName>{ FName(TEXT("FloatArray")), FConcertPropertyChain::InternalContainerPropertyValueName });
-	TestTrue(TEXT("StringArray"), StringArray == TArray<FName>{ TEXT("StringArray") });
-	TestTrue(TEXT("NativeStructArray"), NativeStructArray == TArray<FName>{ TEXT("NativeStructArray") });
-	TestTrue(TEXT("NativeStructArrayInner"), NativeStructArrayInner == TArray<FName>{ FName(TEXT("NativeStructArray")), FConcertPropertyChain::InternalContainerPropertyValueName });
+		TestTrue(TEXT("FloatArray"), FloatArray == TArray<FName>{ TEXT("FloatArray") });
+		TestTrue(TEXT("FloatArrayInner"), FloatArrayInner == TArray<FName>{ FName(TEXT("FloatArray")), FConcertPropertyChain::InternalContainerPropertyValueName });
+		TestTrue(TEXT("StringArray"), StringArray == TArray<FName>{ TEXT("StringArray") });
+		TestTrue(TEXT("NativeStructArray"), NativeStructArray == TArray<FName>{ TEXT("NativeStructArray") });
+		TestTrue(TEXT("NativeStructArrayInner"), NativeStructArrayInner == TArray<FName>{ FName(TEXT("NativeStructArray")), FConcertPropertyChain::InternalContainerPropertyValueName });
 	
-	TestTrue(TEXT("FloatSet"), FloatSet == TArray<FName>{ TEXT("FloatSet") });
-	TestTrue(TEXT("FloatSetInner"), FloatSetInner == TArray<FName>{ FName(TEXT("FloatSet")), FConcertPropertyChain::InternalContainerPropertyValueName });
-	TestTrue(TEXT("StringSet"), StringSet == TArray<FName>{ TEXT("StringSet") });
-	TestTrue(TEXT("NativeStructSet"), NativeStructSet == TArray<FName>{ TEXT("NativeStructSet") });
-	TestTrue(TEXT("NativeStructSetInner"), NativeStructSetInner == TArray<FName>{ FName(TEXT("NativeStructSet")), FConcertPropertyChain::InternalContainerPropertyValueName });
+		TestTrue(TEXT("FloatSet"), FloatSet == TArray<FName>{ TEXT("FloatSet") });
+		TestTrue(TEXT("FloatSetInner"), FloatSetInner == TArray<FName>{ FName(TEXT("FloatSet")), FConcertPropertyChain::InternalContainerPropertyValueName });
+		TestTrue(TEXT("StringSet"), StringSet == TArray<FName>{ TEXT("StringSet") });
+		TestTrue(TEXT("NativeStructSet"), NativeStructSet == TArray<FName>{ TEXT("NativeStructSet") });
+		TestTrue(TEXT("NativeStructSetInner"), NativeStructSetInner == TArray<FName>{ FName(TEXT("NativeStructSet")), FConcertPropertyChain::InternalContainerPropertyValueName });
 	
-	TestTrue(TEXT("StringToFloat"), StringToFloat == TArray<FName>{ TEXT("StringToFloat") });
-	TestTrue(TEXT("StringToFloatInner"), StringToFloatInner == TArray<FName>{ FName(TEXT("StringToFloat")), FConcertPropertyChain::InternalContainerPropertyValueName });
-	TestTrue(TEXT("StringToVector"), StringToVector == TArray<FName>{ TEXT("StringToVector") });
-	TestTrue(TEXT("StringToNativeStruct"), StringToNativeStruct == TArray<FName>{ TEXT("StringToNativeStruct") });
-	TestTrue(TEXT("StringToNativeStructInner"), StringToNativeStructInner == TArray<FName>{ FName(TEXT("StringToNativeStruct")), FConcertPropertyChain::InternalContainerPropertyValueName });
+		TestTrue(TEXT("StringToFloat"), StringToFloat == TArray<FName>{ TEXT("StringToFloat") });
+		TestTrue(TEXT("StringToFloatInner"), StringToFloatInner == TArray<FName>{ FName(TEXT("StringToFloat")), FConcertPropertyChain::InternalContainerPropertyValueName });
+		TestTrue(TEXT("StringToVector"), StringToVector == TArray<FName>{ TEXT("StringToVector") });
+		TestTrue(TEXT("StringToNativeStruct"), StringToNativeStruct == TArray<FName>{ TEXT("StringToNativeStruct") });
+		TestTrue(TEXT("StringToNativeStructInner"), StringToNativeStructInner == TArray<FName>{ FName(TEXT("StringToNativeStruct")), FConcertPropertyChain::InternalContainerPropertyValueName });
 	
-	return true;
-}
-
-/**
- * Tests that all property cases on the specially constructed example class UTestReflectionObject.
- */
-IMPLEMENT_SIMPLE_AUTOMATION_TEST(FConcertForEachReplictableConcertPropertyTests, "Concert.Replication.Data.ForEachReplicatableConcertProperty", EAutomationTestFlags::EditorContext | EAutomationTestFlags::EngineFilter)
-bool FConcertForEachReplictableConcertPropertyTests::RunTest(const FString& Parameters)
-{
-	UClass* Class = UTestReflectionObject::StaticClass();
-	TArray<FConcertPropertyChain> Properties;
-	UE::ConcertSyncCore::PropertyChain::ForEachReplicatableConcertProperty(*Class, [&Properties](FConcertPropertyChain&& Chain)
-	{
-		Properties.Add(MoveTemp(Chain));
-		return EBreakBehavior::Continue;
-	});
+		return true;
+	}
 
-	using FExpectedPropertyPath = TArray<FName>;
-	auto ToString = [](const FExpectedPropertyPath& Path)
+	/**
+	 * Tests that all property cases on the specially constructed example class UTestReflectionObject.
+	 */
+	IMPLEMENT_SIMPLE_AUTOMATION_TEST(FConcertForEachReplictableConcertPropertyTests, "Concert.Replication.Data.ForEachReplicatableConcertProperty", EAutomationTestFlags::EditorContext | EAutomationTestFlags::EngineFilter)
+	bool FConcertForEachReplictableConcertPropertyTests::RunTest(const FString& Parameters)
 	{
-		FString Result;
-		for (int32 i = 0; i < Path.Num(); ++i)
+		UClass* Class = UTestReflectionObject::StaticClass();
+		TArray<FConcertPropertyChain> Properties;
+		UE::ConcertSyncCore::PropertyChain::ForEachReplicatableConcertProperty(*Class, [&Properties](FConcertPropertyChain&& Chain)
 		{
-			if (i == 0)
-			{
-				Result += Path[i].ToString();
-			}
-			else
+			Properties.Add(MoveTemp(Chain));
+			return EBreakBehavior::Continue;
+		});
+
+		using FExpectedPropertyPath = TArray<FName>;
+		const TArray<FExpectedPropertyPath> ExpectedProperties {
+			{ TEXT("Float") },
+			{ TEXT("Vector") },
+			{ TEXT("Vector"), TEXT("X") },
+			{ TEXT("Vector"), TEXT("Y") },
+			{ TEXT("Vector"), TEXT("Z") },
+			{ TEXT("TestStruct") },
+			{ TEXT("TestStruct"), TEXT("Nested") },
+			{ TEXT("TestStruct"), TEXT("Nested"), TEXT("Value") },
+			{ TEXT("TestStruct"), TEXT("Nested"), TEXT("Vector") },
+			{ TEXT("TestStruct"), TEXT("Nested"), TEXT("Vector"), TEXT("X")  },
+			{ TEXT("TestStruct"), TEXT("Nested"), TEXT("Vector"), TEXT("Y")  },
+			{ TEXT("TestStruct"), TEXT("Nested"), TEXT("Vector"), TEXT("Z")  },
+			{ TEXT("TestStruct"), TEXT("Nested"), TEXT("NativeStruct") },
+			{ TEXT("TestStruct"), TEXT("Nested"), TEXT("FloatArray") },
+			{ TEXT("TestStruct"), TEXT("Nested"), TEXT("FloatArray"), FConcertPropertyChain::InternalContainerPropertyValueName },
+			{ TEXT("TestStruct"), TEXT("Nested"), TEXT("StringArray") },
+			{ TEXT("TestStruct"), TEXT("Nested"), TEXT("NativeStructArray") },
+			{ TEXT("TestStruct"), TEXT("Nested"), TEXT("NativeStructArray"), FConcertPropertyChain::InternalContainerPropertyValueName },
+			{ TEXT("TestStruct"), TEXT("Nested"), TEXT("FloatSet") },
+			{ TEXT("TestStruct"), TEXT("Nested"), TEXT("FloatSet"), FConcertPropertyChain::InternalContainerPropertyValueName },
+			{ TEXT("TestStruct"), TEXT("Nested"), TEXT("StringSet") },
+			{ TEXT("TestStruct"), TEXT("Nested"), TEXT("NativeStructSet") },
+			{ TEXT("TestStruct"), TEXT("Nested"), TEXT("NativeStructSet"), FConcertPropertyChain::InternalContainerPropertyValueName },
+			{ TEXT("TestStruct"), TEXT("Nested"), TEXT("StringToFloat") },
+			{ TEXT("TestStruct"), TEXT("Nested"), TEXT("StringToFloat"), FConcertPropertyChain::InternalContainerPropertyValueName},
+			{ TEXT("TestStruct"), TEXT("Nested"), TEXT("StringToVector") },
+			{ TEXT("TestStruct"), TEXT("Nested"), TEXT("StringToVector"), TEXT("X") },
+			{ TEXT("TestStruct"), TEXT("Nested"), TEXT("StringToVector"), TEXT("Y") },
+			{ TEXT("TestStruct"), TEXT("Nested"), TEXT("StringToVector"), TEXT("Z") },
+			{ TEXT("TestStruct"), TEXT("Nested"), TEXT("StringToNativeStruct") },
+			{ TEXT("TestStruct"), TEXT("Nested"), TEXT("StringToNativeStruct"), FConcertPropertyChain::InternalContainerPropertyValueName },
+			// The same again for NestedArray as for Nested just above 
+			{ TEXT("TestStruct"), TEXT("NestedArray") },
+			{ TEXT("TestStruct"), TEXT("NestedArray"), TEXT("Value") },
+			{ TEXT("TestStruct"), TEXT("NestedArray"), TEXT("Vector") },
+			{ TEXT("TestStruct"), TEXT("NestedArray"), TEXT("Vector"), TEXT("X")  },
+			{ TEXT("TestStruct"), TEXT("NestedArray"), TEXT("Vector"), TEXT("Y")  },
+			{ TEXT("TestStruct"), TEXT("NestedArray"), TEXT("Vector"), TEXT("Z")  },
+			{ TEXT("TestStruct"), TEXT("NestedArray"), TEXT("NativeStruct") },
+			{ TEXT("TestStruct"), TEXT("NestedArray"), TEXT("FloatArray") },
+			{ TEXT("TestStruct"), TEXT("NestedArray"), TEXT("FloatArray"), FConcertPropertyChain::InternalContainerPropertyValueName },
+			{ TEXT("TestStruct"), TEXT("NestedArray"), TEXT("StringArray") },
+			{ TEXT("TestStruct"), TEXT("NestedArray"), TEXT("NativeStructArray") },
+			{ TEXT("TestStruct"), TEXT("NestedArray"), TEXT("NativeStructArray"), FConcertPropertyChain::InternalContainerPropertyValueName },
+			{ TEXT("TestStruct"), TEXT("NestedArray"), TEXT("FloatSet") },
+			{ TEXT("TestStruct"), TEXT("NestedArray"), TEXT("FloatSet"), FConcertPropertyChain::InternalContainerPropertyValueName },
+			{ TEXT("TestStruct"), TEXT("NestedArray"), TEXT("StringSet") },
+			{ TEXT("TestStruct"), TEXT("NestedArray"), TEXT("NativeStructSet") },
+			{ TEXT("TestStruct"), TEXT("NestedArray"), TEXT("NativeStructSet"), FConcertPropertyChain::InternalContainerPropertyValueName },
+			{ TEXT("TestStruct"), TEXT("NestedArray"), TEXT("StringToFloat") },
+			{ TEXT("TestStruct"), TEXT("NestedArray"), TEXT("StringToFloat"), FConcertPropertyChain::InternalContainerPropertyValueName},
+			{ TEXT("TestStruct"), TEXT("NestedArray"), TEXT("StringToVector") },
+			{ TEXT("TestStruct"), TEXT("NestedArray"), TEXT("StringToVector"), TEXT("X") },
+			{ TEXT("TestStruct"), TEXT("NestedArray"), TEXT("StringToVector"), TEXT("Y") },
+			{ TEXT("TestStruct"), TEXT("NestedArray"), TEXT("StringToVector"), TEXT("Z") },
+			{ TEXT("TestStruct"), TEXT("NestedArray"), TEXT("StringToNativeStruct") },
+			{ TEXT("TestStruct"), TEXT("NestedArray"), TEXT("StringToNativeStruct"), FConcertPropertyChain::InternalContainerPropertyValueName }
+		};
+
+		for (const FExpectedPropertyPath& Expected : ExpectedProperties)
+		{
+			// This structure is easier for setting breakpoints than using TestTrue
+			if (!Properties.Contains(Expected))
 			{
-				Result += TEXT(".") + Path[i].ToString();
+				AddError(FString::JoinBy(Expected, TEXT("."), [](FName Name){ return Name.ToString(); }));
 			}
 		}
-		return Result;
-	};
-	const TArray<FExpectedPropertyPath> ExpectedProperties {
-		{ TEXT("Float") },
-		{ TEXT("Vector") },
-		{ TEXT("Vector"), TEXT("X") },
-		{ TEXT("Vector"), TEXT("Y") },
-		{ TEXT("Vector"), TEXT("Z") },
-		{ TEXT("TestStruct") },
-		{ TEXT("TestStruct"), TEXT("Nested") },
-		{ TEXT("TestStruct"), TEXT("Nested"), TEXT("Value") },
-		{ TEXT("TestStruct"), TEXT("Nested"), TEXT("Vector") },
-		{ TEXT("TestStruct"), TEXT("Nested"), TEXT("Vector"), TEXT("X")  },
-		{ TEXT("TestStruct"), TEXT("Nested"), TEXT("Vector"), TEXT("Y")  },
-		{ TEXT("TestStruct"), TEXT("Nested"), TEXT("Vector"), TEXT("Z")  },
-		{ TEXT("TestStruct"), TEXT("Nested"), TEXT("NativeStruct") },
-		{ TEXT("TestStruct"), TEXT("Nested"), TEXT("FloatArray") },
-		{ TEXT("TestStruct"), TEXT("Nested"), TEXT("FloatArray"), FConcertPropertyChain::InternalContainerPropertyValueName },
-		{ TEXT("TestStruct"), TEXT("Nested"), TEXT("StringArray") },
-		{ TEXT("TestStruct"), TEXT("Nested"), TEXT("NativeStructArray") },
-		{ TEXT("TestStruct"), TEXT("Nested"), TEXT("NativeStructArray"), FConcertPropertyChain::InternalContainerPropertyValueName },
-		{ TEXT("TestStruct"), TEXT("Nested"), TEXT("FloatSet") },
-		{ TEXT("TestStruct"), TEXT("Nested"), TEXT("FloatSet"), FConcertPropertyChain::InternalContainerPropertyValueName },
-		{ TEXT("TestStruct"), TEXT("Nested"), TEXT("StringSet") },
-		{ TEXT("TestStruct"), TEXT("Nested"), TEXT("NativeStructSet") },
-		{ TEXT("TestStruct"), TEXT("Nested"), TEXT("NativeStructSet"), FConcertPropertyChain::InternalContainerPropertyValueName },
-		{ TEXT("TestStruct"), TEXT("Nested"), TEXT("StringToFloat") },
-		{ TEXT("TestStruct"), TEXT("Nested"), TEXT("StringToFloat"), FConcertPropertyChain::InternalContainerPropertyValueName},
-		{ TEXT("TestStruct"), TEXT("Nested"), TEXT("StringToVector") },
-		{ TEXT("TestStruct"), TEXT("Nested"), TEXT("StringToVector"), TEXT("X") },
-		{ TEXT("TestStruct"), TEXT("Nested"), TEXT("StringToVector"), TEXT("Y") },
-		{ TEXT("TestStruct"), TEXT("Nested"), TEXT("StringToVector"), TEXT("Z") },
-		{ TEXT("TestStruct"), TEXT("Nested"), TEXT("StringToNativeStruct") },
-		{ TEXT("TestStruct"), TEXT("Nested"), TEXT("StringToNativeStruct"), FConcertPropertyChain::InternalContainerPropertyValueName },
-		// The same again for NestedArray as for Nested just above 
-		{ TEXT("TestStruct"), TEXT("NestedArray") },
-		{ TEXT("TestStruct"), TEXT("NestedArray"), TEXT("Value") },
-		{ TEXT("TestStruct"), TEXT("NestedArray"), TEXT("Vector") },
-		{ TEXT("TestStruct"), TEXT("NestedArray"), TEXT("Vector"), TEXT("X")  },
-		{ TEXT("TestStruct"), TEXT("NestedArray"), TEXT("Vector"), TEXT("Y")  },
-		{ TEXT("TestStruct"), TEXT("NestedArray"), TEXT("Vector"), TEXT("Z")  },
-		{ TEXT("TestStruct"), TEXT("NestedArray"), TEXT("NativeStruct") },
-		{ TEXT("TestStruct"), TEXT("NestedArray"), TEXT("FloatArray") },
-		{ TEXT("TestStruct"), TEXT("NestedArray"), TEXT("FloatArray"), FConcertPropertyChain::InternalContainerPropertyValueName },
-		{ TEXT("TestStruct"), TEXT("NestedArray"), TEXT("StringArray") },
-		{ TEXT("TestStruct"), TEXT("NestedArray"), TEXT("NativeStructArray") },
-		{ TEXT("TestStruct"), TEXT("NestedArray"), TEXT("NativeStructArray"), FConcertPropertyChain::InternalContainerPropertyValueName },
-		{ TEXT("TestStruct"), TEXT("NestedArray"), TEXT("FloatSet") },
-		{ TEXT("TestStruct"), TEXT("NestedArray"), TEXT("FloatSet"), FConcertPropertyChain::InternalContainerPropertyValueName },
-		{ TEXT("TestStruct"), TEXT("NestedArray"), TEXT("StringSet") },
-		{ TEXT("TestStruct"), TEXT("NestedArray"), TEXT("NativeStructSet") },
-		{ TEXT("TestStruct"), TEXT("NestedArray"), TEXT("NativeStructSet"), FConcertPropertyChain::InternalContainerPropertyValueName },
-		{ TEXT("TestStruct"), TEXT("NestedArray"), TEXT("StringToFloat") },
-		{ TEXT("TestStruct"), TEXT("NestedArray"), TEXT("StringToFloat"), FConcertPropertyChain::InternalContainerPropertyValueName},
-		{ TEXT("TestStruct"), TEXT("NestedArray"), TEXT("StringToVector") },
-		{ TEXT("TestStruct"), TEXT("NestedArray"), TEXT("StringToVector"), TEXT("X") },
-		{ TEXT("TestStruct"), TEXT("NestedArray"), TEXT("StringToVector"), TEXT("Y") },
-		{ TEXT("TestStruct"), TEXT("NestedArray"), TEXT("StringToVector"), TEXT("Z") },
-		{ TEXT("TestStruct"), TEXT("NestedArray"), TEXT("StringToNativeStruct") },
-		{ TEXT("TestStruct"), TEXT("NestedArray"), TEXT("StringToNativeStruct"), FConcertPropertyChain::InternalContainerPropertyValueName }
-	};
-
-	for (const FExpectedPropertyPath& Expected : ExpectedProperties)
-	{
-		// This structure is easier for setting breakpoints than using TestTrue
-		if (!Properties.Contains(Expected))
-		{
-			AddError(ToString(Expected));
-		}
-	}
 	
-	if (Properties.Num() != ExpectedProperties.Num())
-	{
-		AddError(TEXT("More properties found than expected"));
-		for (const FConcertPropertyChain& Chain : Properties)
+		if (Properties.Num() != ExpectedProperties.Num())
 		{
-			if (!ExpectedProperties.Contains(Chain))
+			AddError(TEXT("More properties found than expected"));
+			for (const FConcertPropertyChain& Chain : Properties)
 			{
-				AddError(FString::Printf(TEXT("Found unexpected chain %s"), *Chain.ToString()));
+				if (!ExpectedProperties.Contains(Chain))
+				{
+					AddError(FString::Printf(TEXT("Found unexpected chain %s"), *Chain.ToString()));
+				}
 			}
 		}
-	}
 	
-	return true;
-}
-
-/**
- * Tests FConcertPropertyChain::Matches.
- */
-IMPLEMENT_SIMPLE_AUTOMATION_TEST(FMatchConcertPropertyChainTest, "Concert.Replication.Data.MatchConcertPropertyChain", EAutomationTestFlags::EditorContext | EAutomationTestFlags::EngineFilter)
-bool FMatchConcertPropertyChainTest::RunTest(const FString& Parameters)
-{
-	// 1. Set up
-	UClass* Class = UTestReflectionObject::StaticClass();
-	FProperty* FloatProperty = Class->FindPropertyByName(GET_MEMBER_NAME_CHECKED(UTestReflectionObject, Float));
-	FProperty* TestStructProperty = Class->FindPropertyByName(GET_MEMBER_NAME_CHECKED(UTestReflectionObject, TestStruct));
-	FProperty* NestedArrayProperty = FTestNestedReplicationStruct::StaticStruct()->FindPropertyByName(GET_MEMBER_NAME_CHECKED(FTestNestedReplicationStruct, NestedArray));
-	FProperty* InternalNestedArrayProperty = CastField<FArrayProperty>(NestedArrayProperty)->Inner;
-	FProperty* ValueProperty = FTestReplicationStruct::StaticStruct()->FindPropertyByName(GET_MEMBER_NAME_CHECKED(FTestReplicationStruct, Value));
-	FProperty* FloatArrayProperty = FTestReplicationStruct::StaticStruct()->FindPropertyByName(GET_MEMBER_NAME_CHECKED(FTestReplicationStruct, FloatArray));
-	FProperty* InternalFloatArrayProperty = CastField<FArrayProperty>(FloatArrayProperty)->Inner;
-	FProperty* StringToFloatProperty = FTestReplicationStruct::StaticStruct()->FindPropertyByName(GET_MEMBER_NAME_CHECKED(FTestReplicationStruct, StringToFloat));
-	FProperty* StringToFloatValueProperty = CastField<FMapProperty>(StringToFloatProperty)->ValueProp;
+		return true;
+	}
+
+	/**
+	 * Tests FConcertPropertyChain::Matches.
+	 */
+	IMPLEMENT_SIMPLE_AUTOMATION_TEST(FMatchConcertPropertyChainTest, "Concert.Replication.Data.MatchConcertPropertyChain", EAutomationTestFlags::EditorContext | EAutomationTestFlags::EngineFilter)
+	bool FMatchConcertPropertyChainTest::RunTest(const FString& Parameters)
+	{
+		// 1. Set up
+		UClass* Class = UTestReflectionObject::StaticClass();
+		FProperty* FloatProperty = Class->FindPropertyByName(GET_MEMBER_NAME_CHECKED(UTestReflectionObject, Float));
+		FProperty* TestStructProperty = Class->FindPropertyByName(GET_MEMBER_NAME_CHECKED(UTestReflectionObject, TestStruct));
+		FProperty* NestedArrayProperty = FTestNestedReplicationStruct::StaticStruct()->FindPropertyByName(GET_MEMBER_NAME_CHECKED(FTestNestedReplicationStruct, NestedArray));
+		FProperty* InternalNestedArrayProperty = CastField<FArrayProperty>(NestedArrayProperty)->Inner;
+		FProperty* ValueProperty = FTestReplicationStruct::StaticStruct()->FindPropertyByName(GET_MEMBER_NAME_CHECKED(FTestReplicationStruct, Value));
+		FProperty* FloatArrayProperty = FTestReplicationStruct::StaticStruct()->FindPropertyByName(GET_MEMBER_NAME_CHECKED(FTestReplicationStruct, FloatArray));
+		FProperty* InternalFloatArrayProperty = CastField<FArrayProperty>(FloatArrayProperty)->Inner;
+		FProperty* StringToFloatProperty = FTestReplicationStruct::StaticStruct()->FindPropertyByName(GET_MEMBER_NAME_CHECKED(FTestReplicationStruct, StringToFloat));
+		FProperty* StringToFloatValueProperty = CastField<FMapProperty>(StringToFloatProperty)->ValueProp;
 	
-	FArchiveSerializedPropertyChain ChainToNestedArrayProperty;
-	ChainToNestedArrayProperty.PushProperty(TestStructProperty, false);
-	// FTestReplicationStruct::Value (float)
-	FArchiveSerializedPropertyChain ChainToNestedProperty = ChainToNestedArrayProperty;
-	// FTestReplicationStruct::FloatArray (TArray<float>)
-	ChainToNestedProperty.PushProperty(NestedArrayProperty, false);
-	ChainToNestedProperty.PushProperty(InternalNestedArrayProperty, false);
-	FArchiveSerializedPropertyChain ChainToFloatArrayProperty = ChainToNestedProperty;
-	FArchiveSerializedPropertyChain ChainToInternalFloatArrayProperty = ChainToFloatArrayProperty;
-	ChainToInternalFloatArrayProperty.PushProperty(FloatArrayProperty, false);
-	// FTestReplicationStruct::StringToFloat (TMap<FString, float>)
-	FArchiveSerializedPropertyChain ChainToStringToFloatProperty = ChainToNestedProperty;
-	FArchiveSerializedPropertyChain ChainToStringToFloatValueProperty = ChainToNestedProperty;
-	ChainToStringToFloatValueProperty.PushProperty(StringToFloatProperty, false);
+		FArchiveSerializedPropertyChain ChainToNestedArrayProperty;
+		ChainToNestedArrayProperty.PushProperty(TestStructProperty, false);
+		// FTestReplicationStruct::Value (float)
+		FArchiveSerializedPropertyChain ChainToNestedProperty = ChainToNestedArrayProperty;
+		// FTestReplicationStruct::FloatArray (TArray<float>)
+		ChainToNestedProperty.PushProperty(NestedArrayProperty, false);
+		ChainToNestedProperty.PushProperty(InternalNestedArrayProperty, false);
+		FArchiveSerializedPropertyChain ChainToFloatArrayProperty = ChainToNestedProperty;
+		FArchiveSerializedPropertyChain ChainToInternalFloatArrayProperty = ChainToFloatArrayProperty;
+		ChainToInternalFloatArrayProperty.PushProperty(FloatArrayProperty, false);
+		// FTestReplicationStruct::StringToFloat (TMap<FString, float>)
+		FArchiveSerializedPropertyChain ChainToStringToFloatProperty = ChainToNestedProperty;
+		FArchiveSerializedPropertyChain ChainToStringToFloatValueProperty = ChainToNestedProperty;
+		ChainToStringToFloatValueProperty.PushProperty(StringToFloatProperty, false);
 	
-	const FConcertPropertyChain ConcertChain_Simple_FloatProperty(nullptr, *FloatProperty);
-	const FConcertPropertyChain ConcertChain_Nested_ArrayProperty(&ChainToNestedArrayProperty, *NestedArrayProperty);
-	const FConcertPropertyChain ConcertChain_Nested_ValueProperty(&ChainToNestedProperty, *ValueProperty);
-	const FConcertPropertyChain ConcertChain_Nested_FloatArrayProperty(&ChainToFloatArrayProperty, *FloatArrayProperty);
-	const FConcertPropertyChain ConcertChain_Nested_FloatArrayInnerProperty(&ChainToInternalFloatArrayProperty, *InternalFloatArrayProperty);
-	const FConcertPropertyChain ConcertChain_Nested_StringToFloatProperty(&ChainToStringToFloatProperty, *StringToFloatProperty);
-	const FConcertPropertyChain ConcertChain_Nested_StringToFloatValueProperty(&ChainToStringToFloatValueProperty, *StringToFloatValueProperty);
+		const FConcertPropertyChain ConcertChain_Simple_FloatProperty(nullptr, *FloatProperty);
+		const FConcertPropertyChain ConcertChain_Nested_ArrayProperty(&ChainToNestedArrayProperty, *NestedArrayProperty);
+		const FConcertPropertyChain ConcertChain_Nested_ValueProperty(&ChainToNestedProperty, *ValueProperty);
+		const FConcertPropertyChain ConcertChain_Nested_FloatArrayProperty(&ChainToFloatArrayProperty, *FloatArrayProperty);
+		const FConcertPropertyChain ConcertChain_Nested_FloatArrayInnerProperty(&ChainToInternalFloatArrayProperty, *InternalFloatArrayProperty);
+		const FConcertPropertyChain ConcertChain_Nested_StringToFloatProperty(&ChainToStringToFloatProperty, *StringToFloatProperty);
+		const FConcertPropertyChain ConcertChain_Nested_StringToFloatValueProperty(&ChainToStringToFloatValueProperty, *StringToFloatValueProperty);
 	
 
 
-	// 2. Run
-	const bool bMatches_Simple_FloatProperty = ConcertChain_Simple_FloatProperty.MatchesExactly(nullptr, *FloatProperty);
-	const bool bMatches_Nested_ArrayProperty = ConcertChain_Nested_ArrayProperty.MatchesExactly(&ChainToNestedArrayProperty, *NestedArrayProperty);
-	const bool bMatches_Nested_Property = ConcertChain_Nested_ValueProperty.MatchesExactly(&ChainToNestedProperty, *ValueProperty);
-	const bool bMatches_Nested_FloatArrayProperty = ConcertChain_Nested_FloatArrayProperty.MatchesExactly(&ChainToFloatArrayProperty, *FloatArrayProperty);
-	const bool bMatches_Nested_FloatArrayInnerProperty = ConcertChain_Nested_FloatArrayInnerProperty.MatchesExactly(&ChainToInternalFloatArrayProperty, *InternalFloatArrayProperty);
-	const bool bMatches_Nested_StringToFloatProperty = ConcertChain_Nested_StringToFloatProperty.MatchesExactly(&ChainToStringToFloatProperty, *StringToFloatProperty);
-	const bool bMatches_Nested_StringToFloatValueProperty = ConcertChain_Nested_StringToFloatValueProperty.MatchesExactly(&ChainToStringToFloatValueProperty, *StringToFloatValueProperty);
-
-
-
-
-	// 3. Test
-	TestTrue(TEXT("Float"), bMatches_Simple_FloatProperty);
-	TestTrue(TEXT("NestedArray"), bMatches_Nested_ArrayProperty);
-	// This case is interesting because e.g. ChainToNestedProperty = { NestedArray, NestedArray } but ConcertChain_NestedProperty = { NestedArray, Value } (set a breakpoint to convince yourself)
-	TestTrue(TEXT("NestedArray.NestedArray.Value"), bMatches_Nested_Property);
-	// This would restore the size of the array only...
-	TestTrue(TEXT("NestedArray.NestedArray.FloatArray"), bMatches_Nested_FloatArrayProperty);
-	// ... and this includes the float values. This one is interesting because it ends with FConcertPropertyChain::InternalContainerPropertyValueName
-	TestTrue(TEXT("NestedArray.NestedArray.FloatArray.Value (ends with FConcertPropertyChain::InternalContainerPropertyValueName)"), bMatches_Nested_FloatArrayInnerProperty);
-	// This would restore the keys and the size of the map only...
-	TestTrue(TEXT("NestedArray.NestedArray.StringToFloat"), bMatches_Nested_StringToFloatProperty);
-	// ... and this includes the float values. This one is interesting because it ends with FConcertPropertyChain::InternalContainerPropertyValueName
-	TestTrue(TEXT("NestedArray.NestedArray.StringToFloat.Value (ends with FConcertPropertyChain::InternalContainerPropertyValueName)"), bMatches_Nested_StringToFloatValueProperty);
-
-	return true;
-}
-
-/**
- * Tests that UE::ConcertSyncCore::FReplicationPropertyFilter works correctly.
- * This uses the filter and serializes an object, doing a Memcmp to detect that exactly the expected properties are serialized.
- */
-IMPLEMENT_SIMPLE_AUTOMATION_TEST(FSerializeWithPropertyFilterTest, "Concert.Replication.Data.SerializeWithPropertyFilter", EAutomationTestFlags::EditorContext | EAutomationTestFlags::EngineFilter)
-bool FSerializeWithPropertyFilterTest::RunTest(const FString& Parameters)
-{
-	class FTestArchive : public FMemoryWriter
-	{
-		FAutomationTestBase& Test;
-		const FConcertPropertySelection& Selection;
-		const UE::ConcertSyncCore::FReplicationPropertyFilter PropertyFilter;
-	public:
-		
-		TArray<uint8> Bytes;
+		// 2. Run
+		const bool bMatches_Simple_FloatProperty = ConcertChain_Simple_FloatProperty.MatchesExactly(nullptr, *FloatProperty);
+		const bool bMatches_Nested_ArrayProperty = ConcertChain_Nested_ArrayProperty.MatchesExactly(&ChainToNestedArrayProperty, *NestedArrayProperty);
+		const bool bMatches_Nested_Property = ConcertChain_Nested_ValueProperty.MatchesExactly(&ChainToNestedProperty, *ValueProperty);
+		const bool bMatches_Nested_FloatArrayProperty = ConcertChain_Nested_FloatArrayProperty.MatchesExactly(&ChainToFloatArrayProperty, *FloatArrayProperty);
+		const bool bMatches_Nested_FloatArrayInnerProperty = ConcertChain_Nested_FloatArrayInnerProperty.MatchesExactly(&ChainToInternalFloatArrayProperty, *InternalFloatArrayProperty);
+		const bool bMatches_Nested_StringToFloatProperty = ConcertChain_Nested_StringToFloatProperty.MatchesExactly(&ChainToStringToFloatProperty, *StringToFloatProperty);
+		const bool bMatches_Nested_StringToFloatValueProperty = ConcertChain_Nested_StringToFloatValueProperty.MatchesExactly(&ChainToStringToFloatValueProperty, *StringToFloatValueProperty);
+
 
-		FTestArchive(FAutomationTestBase& Test, const FConcertPropertySelection& Selection)
-			: FMemoryWriter(Bytes)
-			, Test(Test)
-			, Selection(Selection)
-			, PropertyFilter(Selection)
-		{
-			FArchive::SetIsSaving(true);
-		}
 
-		virtual bool ShouldSkipProperty(const FProperty* InProperty) const override
+
+		// 3. Test
+		TestTrue(TEXT("Float"), bMatches_Simple_FloatProperty);
+		TestTrue(TEXT("NestedArray"), bMatches_Nested_ArrayProperty);
+		// This case is interesting because e.g. ChainToNestedProperty = { NestedArray, NestedArray } but ConcertChain_NestedProperty = { NestedArray, Value } (set a breakpoint to convince yourself)
+		TestTrue(TEXT("NestedArray.NestedArray.Value"), bMatches_Nested_Property);
+		// This would restore the size of the array only...
+		TestTrue(TEXT("NestedArray.NestedArray.FloatArray"), bMatches_Nested_FloatArrayProperty);
+		// ... and this includes the float values. This one is interesting because it ends with FConcertPropertyChain::InternalContainerPropertyValueName
+		TestTrue(TEXT("NestedArray.NestedArray.FloatArray.Value (ends with FConcertPropertyChain::InternalContainerPropertyValueName)"), bMatches_Nested_FloatArrayInnerProperty);
+		// This would restore the keys and the size of the map only...
+		TestTrue(TEXT("NestedArray.NestedArray.StringToFloat"), bMatches_Nested_StringToFloatProperty);
+		// ... and this includes the float values. This one is interesting because it ends with FConcertPropertyChain::InternalContainerPropertyValueName
+		TestTrue(TEXT("NestedArray.NestedArray.StringToFloat.Value (ends with FConcertPropertyChain::InternalContainerPropertyValueName)"), bMatches_Nested_StringToFloatValueProperty);
+
+		return true;
+	}
+
+	/**
+	 * Tests that UE::ConcertSyncCore::FReplicationPropertyFilter works correctly.
+	 * This uses the filter and serializes an object, doing a Memcmp to detect that exactly the expected properties are serialized.
+	 */
+	IMPLEMENT_SIMPLE_AUTOMATION_TEST(FSerializeWithPropertyFilterTest, "Concert.Replication.Data.SerializeWithPropertyFilter", EAutomationTestFlags::EditorContext | EAutomationTestFlags::EngineFilter)
+	bool FSerializeWithPropertyFilterTest::RunTest(const FString& Parameters)
+	{
+		class FTestArchive : public FMemoryWriter
 		{
-			if (!UE::ConcertSyncCore::PropertyChain::IsReplicatableProperty(*InProperty))
-			{
-				return true;
-			}
+			FAutomationTestBase& Test;
+			const FConcertPropertySelection& Selection;
+			const UE::ConcertSyncCore::FReplicationPropertyFilter PropertyFilter;
+		public:
+		
+			TArray<uint8> Bytes;
 
-			// We can assume the constructor of FConcertPropertyChain works correctly here since we've tested it in another test.
-			const FConcertPropertyChain Chain(GetSerializedPropertyChain(), *InProperty);
-			const bool bIsPropertyInChain = Selection.ReplicatedProperties.Contains(Chain);
-			const bool bIsAllowed = PropertyFilter.ShouldSerializeProperty(GetSerializedPropertyChain(), *InProperty);
-			if (bIsPropertyInChain)
+			FTestArchive(FAutomationTestBase& Test, const FConcertPropertySelection& Selection)
+				: FMemoryWriter(Bytes)
+				, Test(Test)
+				, Selection(Selection)
+				, PropertyFilter(Selection)
 			{
-				Test.TestTrue(FString::Printf(TEXT("Allowed %s"), *InProperty->GetAuthoredName()), bIsAllowed);
+				FArchive::SetIsSaving(true);
 			}
-			else
+
+			virtual bool ShouldSkipProperty(const FProperty* InProperty) const override
 			{
-				Test.TestFalse(FString::Printf(TEXT("Disallowed %s"), *InProperty->GetAuthoredName()), bIsAllowed);
+				if (!UE::ConcertSyncCore::PropertyChain::IsReplicatableProperty(*InProperty))
+				{
+					return true;
+				}
+
+				// We can assume the constructor of FConcertPropertyChain works correctly here since we've tested it in another test.
+				const FConcertPropertyChain Chain(GetSerializedPropertyChain(), *InProperty);
+				const bool bIsPropertyInChain = Selection.ReplicatedProperties.Contains(Chain);
+				const bool bIsAllowed = PropertyFilter.ShouldSerializeProperty(GetSerializedPropertyChain(), *InProperty);
+				if (bIsPropertyInChain)
+				{
+					Test.TestTrue(FString::Printf(TEXT("Allowed %s"), *InProperty->GetAuthoredName()), bIsAllowed);
+				}
+				else
+				{
+					Test.TestFalse(FString::Printf(TEXT("Disallowed %s"), *InProperty->GetAuthoredName()), bIsAllowed);
+				}
+
+				return !bIsAllowed;
 			}
 
-			return !bIsAllowed;
+			// If we encounter any objects, ring alarm bells because IsReplicatableProperty should have covered this case
+			virtual FArchive& operator<<(UObject*& Obj) override{ Test.AddError(TEXT("Object not expected at this point")); return *this; }
+			virtual FArchive& operator<<(FWeakObjectPtr& Obj) override { Test.AddError(TEXT("Object not expected at this point")); return *this; }
+			virtual FArchive& operator<<(FSoftObjectPtr& Value) override { Test.AddError(TEXT("Object not expected at this point")); return *this; } 
+			virtual FArchive& operator<<(FSoftObjectPath& Value) override { Test.AddError(TEXT("Object not expected at this point")); return *this; }
+			virtual FArchive& operator<<(FObjectPtr& Obj) override { Test.AddError(TEXT("Object not expected at this point")); return *this; }
+		};
+
+		// CreateFromPath is supposed to work (assuming the other tests worked so far) but we'll use TOptional::Get() calls are in case it does not. 
+		TArray<FConcertPropertyChain> AllowedProperties = {
+			FConcertPropertyChain::CreateFromPath(*UTestReflectionObject::StaticClass(), { TEXT("Float") }).Get({}),
+			FConcertPropertyChain::CreateFromPath(*UTestReflectionObject::StaticClass(), { TEXT("Vector") }).Get({}),
+			FConcertPropertyChain::CreateFromPath(*UTestReflectionObject::StaticClass(), { TEXT("Vector"), TEXT("X") }).Get({}),
+			FConcertPropertyChain::CreateFromPath(*UTestReflectionObject::StaticClass(), { TEXT("Vector"), TEXT("Y") }).Get({}),
+			FConcertPropertyChain::CreateFromPath(*UTestReflectionObject::StaticClass(), { TEXT("Vector"), TEXT("Z") }).Get({}),
+			FConcertPropertyChain::CreateFromPath(*UTestReflectionObject::StaticClass(), { TEXT("TestStruct") }).Get({}),
+			FConcertPropertyChain::CreateFromPath(*UTestReflectionObject::StaticClass(), { TEXT("TestStruct"), TEXT("Nested") }).Get({}),
+			FConcertPropertyChain::CreateFromPath(*UTestReflectionObject::StaticClass(), { TEXT("TestStruct"), TEXT("Nested"), TEXT("Value") }).Get({}),
+			FConcertPropertyChain::CreateFromPath(*UTestReflectionObject::StaticClass(), { TEXT("TestStruct"), TEXT("Nested"), TEXT("NativeStruct") }).Get({}),
+			FConcertPropertyChain::CreateFromPath(*UTestReflectionObject::StaticClass(), { TEXT("TestStruct"), TEXT("Nested"), TEXT("FloatArray") }).Get({}),
+			FConcertPropertyChain::CreateFromPath(*UTestReflectionObject::StaticClass(), { TEXT("TestStruct"), TEXT("Nested"), TEXT("FloatArray"), FConcertPropertyChain::InternalContainerPropertyValueName }).Get({}),
+			FConcertPropertyChain::CreateFromPath(*UTestReflectionObject::StaticClass(), { TEXT("TestStruct"), TEXT("NestedArray") }).Get({}),
+			FConcertPropertyChain::CreateFromPath(*UTestReflectionObject::StaticClass(), { TEXT("TestStruct"), TEXT("NestedArray"), TEXT("Value") }).Get({}),
+			FConcertPropertyChain::CreateFromPath(*UTestReflectionObject::StaticClass(), { TEXT("TestStruct"), TEXT("NestedArray"), TEXT("Vector") }).Get({}),
+			FConcertPropertyChain::CreateFromPath(*UTestReflectionObject::StaticClass(), { TEXT("TestStruct"), TEXT("NestedArray"), TEXT("Vector"), TEXT("X") }).Get({}),
+			FConcertPropertyChain::CreateFromPath(*UTestReflectionObject::StaticClass(), { TEXT("TestStruct"), TEXT("NestedArray"), TEXT("Vector"), TEXT("Y") }).Get({}),
+			FConcertPropertyChain::CreateFromPath(*UTestReflectionObject::StaticClass(), { TEXT("TestStruct"), TEXT("NestedArray"), TEXT("Vector"), TEXT("Z") }).Get({}),
+			FConcertPropertyChain::CreateFromPath(*UTestReflectionObject::StaticClass(), { TEXT("TestStruct"), TEXT("NestedArray"), TEXT("NativeStruct") }).Get({}),
+			FConcertPropertyChain::CreateFromPath(*UTestReflectionObject::StaticClass(), { TEXT("TestStruct"), TEXT("NestedArray"), TEXT("FloatArray") }).Get({}),
+			FConcertPropertyChain::CreateFromPath(*UTestReflectionObject::StaticClass(), { TEXT("TestStruct"), TEXT("NestedArray"), TEXT("FloatArray"), FConcertPropertyChain::InternalContainerPropertyValueName }).Get({})
+		};
+		if (AllowedProperties.Contains(FConcertPropertyChain{}))
+		{
+			AddError(TEXT("FConcertPropertyChain::CreateFromPath was expected to work on all of the above cases. Take a look at the ForEachReplicatableConcertProperty test, which should have failed, as well."));
+			return false;
 		}
 
-		// If we encounter any objects, ring alarm bells because IsReplicatableProperty should have covered this case
-		virtual FArchive& operator<<(UObject*& Obj) override{ Test.AddError(TEXT("Object not expected at this point")); return *this; }
-		virtual FArchive& operator<<(FWeakObjectPtr& Obj) override { Test.AddError(TEXT("Object not expected at this point")); return *this; }
-		virtual FArchive& operator<<(FSoftObjectPtr& Value) override { Test.AddError(TEXT("Object not expected at this point")); return *this; } 
-		virtual FArchive& operator<<(FSoftObjectPath& Value) override { Test.AddError(TEXT("Object not expected at this point")); return *this; }
-		virtual FArchive& operator<<(FObjectPtr& Obj) override { Test.AddError(TEXT("Object not expected at this point")); return *this; }
-	};
-
-	// CreateFromPath is supposed to work (assuming the other tests worked so far) but we'll use TOptional::Get() calls are in case it does not. 
-	TArray<FConcertPropertyChain> AllowedProperties = {
-		FConcertPropertyChain::CreateFromPath(*UTestReflectionObject::StaticClass(), { TEXT("Float") }).Get({}),
-		FConcertPropertyChain::CreateFromPath(*UTestReflectionObject::StaticClass(), { TEXT("Vector") }).Get({}),
-		FConcertPropertyChain::CreateFromPath(*UTestReflectionObject::StaticClass(), { TEXT("Vector"), TEXT("X") }).Get({}),
-		FConcertPropertyChain::CreateFromPath(*UTestReflectionObject::StaticClass(), { TEXT("Vector"), TEXT("Y") }).Get({}),
-		FConcertPropertyChain::CreateFromPath(*UTestReflectionObject::StaticClass(), { TEXT("Vector"), TEXT("Z") }).Get({}),
-		FConcertPropertyChain::CreateFromPath(*UTestReflectionObject::StaticClass(), { TEXT("TestStruct") }).Get({}),
-		FConcertPropertyChain::CreateFromPath(*UTestReflectionObject::StaticClass(), { TEXT("TestStruct"), TEXT("Nested") }).Get({}),
-		FConcertPropertyChain::CreateFromPath(*UTestReflectionObject::StaticClass(), { TEXT("TestStruct"), TEXT("Nested"), TEXT("Value") }).Get({}),
-		FConcertPropertyChain::CreateFromPath(*UTestReflectionObject::StaticClass(), { TEXT("TestStruct"), TEXT("Nested"), TEXT("NativeStruct") }).Get({}),
-		FConcertPropertyChain::CreateFromPath(*UTestReflectionObject::StaticClass(), { TEXT("TestStruct"), TEXT("Nested"), TEXT("FloatArray") }).Get({}),
-		FConcertPropertyChain::CreateFromPath(*UTestReflectionObject::StaticClass(), { TEXT("TestStruct"), TEXT("Nested"), TEXT("FloatArray"), FConcertPropertyChain::InternalContainerPropertyValueName }).Get({}),
-		FConcertPropertyChain::CreateFromPath(*UTestReflectionObject::StaticClass(), { TEXT("TestStruct"), TEXT("NestedArray") }).Get({}),
-		FConcertPropertyChain::CreateFromPath(*UTestReflectionObject::StaticClass(), { TEXT("TestStruct"), TEXT("NestedArray"), TEXT("Value") }).Get({}),
-		FConcertPropertyChain::CreateFromPath(*UTestReflectionObject::StaticClass(), { TEXT("TestStruct"), TEXT("NestedArray"), TEXT("Vector") }).Get({}),
-		FConcertPropertyChain::CreateFromPath(*UTestReflectionObject::StaticClass(), { TEXT("TestStruct"), TEXT("NestedArray"), TEXT("Vector"), TEXT("X") }).Get({}),
-		FConcertPropertyChain::CreateFromPath(*UTestReflectionObject::StaticClass(), { TEXT("TestStruct"), TEXT("NestedArray"), TEXT("Vector"), TEXT("Y") }).Get({}),
-		FConcertPropertyChain::CreateFromPath(*UTestReflectionObject::StaticClass(), { TEXT("TestStruct"), TEXT("NestedArray"), TEXT("Vector"), TEXT("Z") }).Get({}),
-		FConcertPropertyChain::CreateFromPath(*UTestReflectionObject::StaticClass(), { TEXT("TestStruct"), TEXT("NestedArray"), TEXT("NativeStruct") }).Get({}),
-		FConcertPropertyChain::CreateFromPath(*UTestReflectionObject::StaticClass(), { TEXT("TestStruct"), TEXT("NestedArray"), TEXT("FloatArray") }).Get({}),
-		FConcertPropertyChain::CreateFromPath(*UTestReflectionObject::StaticClass(), { TEXT("TestStruct"), TEXT("NestedArray"), TEXT("FloatArray"), FConcertPropertyChain::InternalContainerPropertyValueName }).Get({})
-	};
-	if (AllowedProperties.Contains(FConcertPropertyChain{}))
-	{
-		AddError(TEXT("FConcertPropertyChain::CreateFromPath was expected to work on all of the above cases. Take a look at the ForEachReplicatableConcertProperty test, which should have failed, as well."));
-		return false;
-	}
-
-	auto SetTestValues = [](UTestReflectionObject& Object)
-	{
-		Object.Float = 42.f;
-		Object.Vector = { 10.f, 20.f, 30.f };
+		auto SetTestValues = [](UTestReflectionObject& Object)
+		{
+			Object.Float = 42.f;
+			Object.Vector = { 10.f, 20.f, 30.f };
 		
-		Object.TestStruct.Nested.Value = 420.f;
-		Object.TestStruct.Nested.NativeStruct.Float = 4200.f;
-		Object.TestStruct.Nested.FloatArray = { 500.f, 750.f };
-
-		Object.TestStruct.NestedArray = { FTestReplicationStruct{} };
-		Object.TestStruct.NestedArray[0].Value = -420.f;
-		Object.TestStruct.NestedArray[0].Vector = { -100.f, -200.f, -300.f };
-		Object.TestStruct.NestedArray[0].NativeStruct.Float = -4200.f;
-		Object.TestStruct.NestedArray[0].FloatArray = { -500.f, -750.f };
-	};
-	auto ResetTestValues = [](UTestReflectionObject& Object)
-	{
-		Object.Float = {};
-		Object.Vector = FVector::ZeroVector;
-		Object.TestStruct = FTestNestedReplicationStruct{};
-	};
-	auto ObjectToBinary = [](UObject& Object)
-	{
-		TArray<uint8> Bytes;
-		FObjectWriter CollectAllData(&Object, Bytes, false);
-		return Bytes;
-	};
-	auto ValidateTestValues = [this](UTestReflectionObject& Object)
-	{
-		TestEqual(TEXT("Object.Float = 42.f"), Object.Float, 42.f);
-		TestEqual(TEXT("Object.Vector = { 10.f, 20.f, 30.f }"), Object.Vector, FVector{10.f, 20.f, 30.f });
+			Object.TestStruct.Nested.Value = 420.f;
+			Object.TestStruct.Nested.NativeStruct.Float = 4200.f;
+			Object.TestStruct.Nested.FloatArray = { 500.f, 750.f };
+
+			Object.TestStruct.NestedArray = { FTestReplicationStruct{} };
+			Object.TestStruct.NestedArray[0].Value = -420.f;
+			Object.TestStruct.NestedArray[0].Vector = { -100.f, -200.f, -300.f };
+			Object.TestStruct.NestedArray[0].NativeStruct.Float = -4200.f;
+			Object.TestStruct.NestedArray[0].FloatArray = { -500.f, -750.f };
+		};
+		auto ResetTestValues = [](UTestReflectionObject& Object)
+		{
+			Object.Float = {};
+			Object.Vector = FVector::ZeroVector;
+			Object.TestStruct = FTestNestedReplicationStruct{};
+		};
+		auto ObjectToBinary = [](UObject& Object)
+		{
+			TArray<uint8> Bytes;
+			FObjectWriter CollectAllData(&Object, Bytes, false);
+			return Bytes;
+		};
+		auto ValidateTestValues = [this](UTestReflectionObject& Object)
+		{
+			TestEqual(TEXT("Object.Float = 42.f"), Object.Float, 42.f);
+			TestEqual(TEXT("Object.Vector = { 10.f, 20.f, 30.f }"), Object.Vector, FVector{10.f, 20.f, 30.f });
 		
-		TestEqual(TEXT("Object.TestStruct.Nested.Value = 420.f"), Object.TestStruct.Nested.Value, 420.f);
-		TestEqual(TEXT("Object.TestStruct.Nested.NativeStruct.Float = 4200.f"), Object.TestStruct.Nested.NativeStruct.Float, 4200.f);
-		TestEqual(TEXT("Object.TestStruct.Nested.FloatArray = { 500.f, 750.f }"), Object.TestStruct.Nested.FloatArray, TArray<float>{ 500.f, 750.f });
-
-		TestEqual(TEXT("Object.TestStruct.NestedArray[0].Value = -420.f"), Object.TestStruct.NestedArray[0].Value, -420.f);
-		TestEqual(TEXT("Object.TestStruct.NestedArray[0].Vector = { -100.f, -200.f, -300.f }"), Object.TestStruct.NestedArray[0].Vector, FVector{ -100.f, -200.f, -300.f });
-		TestEqual(TEXT("Object.TestStruct.NestedArray[0].NativeStruct.Float = -4200.f"), Object.TestStruct.NestedArray[0].NativeStruct.Float, -4200.f);
-		TestEqual(TEXT("Object.TestStruct.NestedArray[0].FloatArray = { -500.f, -750.f }"), Object.TestStruct.NestedArray[0].FloatArray, TArray<float>{ -500.f, -750.f });
-	};
+			TestEqual(TEXT("Object.TestStruct.Nested.Value = 420.f"), Object.TestStruct.Nested.Value, 420.f);
+			TestEqual(TEXT("Object.TestStruct.Nested.NativeStruct.Float = 4200.f"), Object.TestStruct.Nested.NativeStruct.Float, 4200.f);
+			TestEqual(TEXT("Object.TestStruct.Nested.FloatArray = { 500.f, 750.f }"), Object.TestStruct.Nested.FloatArray, TArray<float>{ 500.f, 750.f });
+
+			TestEqual(TEXT("Object.TestStruct.NestedArray[0].Value = -420.f"), Object.TestStruct.NestedArray[0].Value, -420.f);
+			TestEqual(TEXT("Object.TestStruct.NestedArray[0].Vector = { -100.f, -200.f, -300.f }"), Object.TestStruct.NestedArray[0].Vector, FVector{ -100.f, -200.f, -300.f });
+			TestEqual(TEXT("Object.TestStruct.NestedArray[0].NativeStruct.Float = -4200.f"), Object.TestStruct.NestedArray[0].NativeStruct.Float, -4200.f);
+			TestEqual(TEXT("Object.TestStruct.NestedArray[0].FloatArray = { -500.f, -750.f }"), Object.TestStruct.NestedArray[0].FloatArray, TArray<float>{ -500.f, -750.f });
+		};
 	
-	UTestReflectionObject* TestObject = NewObject<UTestReflectionObject>();
-	SetTestValues(*TestObject);
+		UTestReflectionObject* TestObject = NewObject<UTestReflectionObject>();
+		SetTestValues(*TestObject);
 	
-	// This property is supposed to be filtered out...
-	TestObject->TestStruct.Nested.Vector = { 100.f, 200.f, 300.f };
-	const TArray<uint8> SnapshotWithTestValues = ObjectToBinary(*TestObject);
-	// ... and we change its value after taking snapshotting to validate the filter indeed filters out the undesired property...
-	TestObject->TestStruct.Nested.Vector = { 1000.f, 2000.f, 3000.f };
-
-	// This does the filtering & tests that properties are filtered correctly using ShouldSkipProperty
-	const FConcertPropertySelection Selection{ AllowedProperties };
-	FTestArchive SavingArchive(*this, Selection);
-	FNameAsStringProxyArchive CrashProtection_Saving(SavingArchive);
-	TestObject->Serialize(CrashProtection_Saving);
-
-	// ... the supposed to be filtered property must now be reset to the same value as was in the snapshot so Memcmp does not detect it.
-	ResetTestValues(*TestObject);
-	TestObject->TestStruct.Nested.Vector = { 100.f, 200.f, 300.f };
-	FObjectReader LoadingArchive(SavingArchive.Bytes);
-	FNameAsStringProxyArchive CrashProtection_Loading(LoadingArchive);
-	TestObject->Serialize(CrashProtection_Loading);
-	const TArray<uint8> SnapshotAfterLoading = ObjectToBinary(*TestObject);
+		// This property is supposed to be filtered out...
+		TestObject->TestStruct.Nested.Vector = { 100.f, 200.f, 300.f };
+		const TArray<uint8> SnapshotWithTestValues = ObjectToBinary(*TestObject);
+		// ... and we change its value after taking snapshotting to validate the filter indeed filters out the undesired property...
+		TestObject->TestStruct.Nested.Vector = { 1000.f, 2000.f, 3000.f };
+
+		// This does the filtering & tests that properties are filtered correctly using ShouldSkipProperty
+		const FConcertPropertySelection Selection{ AllowedProperties };
+		FTestArchive SavingArchive(*this, Selection);
+		FNameAsStringProxyArchive CrashProtection_Saving(SavingArchive);
+		TestObject->Serialize(CrashProtection_Saving);
+
+		// ... the supposed to be filtered property must now be reset to the same value as was in the snapshot so Memcmp does not detect it.
+		ResetTestValues(*TestObject);
+		TestObject->TestStruct.Nested.Vector = { 100.f, 200.f, 300.f };
+		FObjectReader LoadingArchive(SavingArchive.Bytes);
+		FNameAsStringProxyArchive CrashProtection_Loading(LoadingArchive);
+		TestObject->Serialize(CrashProtection_Loading);
+		const TArray<uint8> SnapshotAfterLoading = ObjectToBinary(*TestObject);
 	
-	// We do a Memcmp on the snapshot data before filtering and on the data after applying to a blanked out object. They must be equal or the filter did not work correctly!
-	if (SnapshotWithTestValues.Num() == SnapshotAfterLoading.Num())
-	{
-		const bool bEqualMemory = FPlatformMemory::Memcmp(SnapshotWithTestValues.GetData(), SnapshotAfterLoading.GetData(), SnapshotWithTestValues.Num()) == 0;
-		TestTrue(TEXT("Memcmp returned non-zero: Applying the filtered data to a reset object did not leave the object in the same state as it was serialized in."), bEqualMemory);
+		// We do a Memcmp on the snapshot data before filtering and on the data after applying to a blanked out object. They must be equal or the filter did not work correctly!
+		if (SnapshotWithTestValues.Num() == SnapshotAfterLoading.Num())
+		{
+			const bool bEqualMemory = FPlatformMemory::Memcmp(SnapshotWithTestValues.GetData(), SnapshotAfterLoading.GetData(), SnapshotWithTestValues.Num()) == 0;
+			TestTrue(TEXT("Memcmp returned non-zero: Applying the filtered data to a reset object did not leave the object in the same state as it was serialized in."), bEqualMemory);
+		}
+		else
+		{
+			AddError(TEXT("BinaryDataWithTestValues != BinaryDataAfterLoading: Applying the filtered data to a reset object did not leave the object in the same state as it was serialized in."));
+		}
+
+		// This is technically unnecessary because Memcmp would already have failed the test above but this helps in debugging what values caused the Memcmp to be different.
+		ValidateTestValues(*TestObject);
+		TestEqual(TEXT("Object.TestStruct.Nested.Vector was left unchanged"), TestObject->TestStruct.Nested.Vector, FVector{ 100.f, 200.f, 300.f });
+		return true;
 	}
-	else
+
+	/**
+	 * Tests that UE::ConcertSyncCore::FReplicationPropertyFilter works correctly.
+	 * This uses the filter and serializes an object, doing a Memcmp to detect that exactly the expected properties are serialized.
+	 */
+	IMPLEMENT_SIMPLE_AUTOMATION_TEST(FOverlapPropertiesTest, "Concert.Replication.Data.OverlapProperties", EAutomationTestFlags::EditorContext | EAutomationTestFlags::EngineFilter)
+	bool FOverlapPropertiesTest::RunTest(const FString& Parameters)
 	{
-		AddError(TEXT("BinaryDataWithTestValues != BinaryDataAfterLoading: Applying the filtered data to a reset object did not leave the object in the same state as it was serialized in."));
-	}
+		// 1. Define paths
+		using FExpectedPropertyPath = TArray<FName>;
+		const TArray<FExpectedPropertyPath> RootProperties {
+			{ TEXT("Float") },
+			{ TEXT("Vector") },
+			{ TEXT("Vector"), TEXT("X") },
+			{ TEXT("Vector"), TEXT("Y") },
+			{ TEXT("Vector"), TEXT("Z") },
+			};
+		const TArray<FExpectedPropertyPath> NestedProperties = {
+			{ TEXT("TestStruct") },
+			{ TEXT("TestStruct"), TEXT("Nested") },
+			{ TEXT("TestStruct"), TEXT("Nested"), TEXT("Value") },
+			{ TEXT("TestStruct"), TEXT("Nested"), TEXT("Vector") },
+			{ TEXT("TestStruct"), TEXT("Nested"), TEXT("Vector"), TEXT("X")  },
+			{ TEXT("TestStruct"), TEXT("Nested"), TEXT("Vector"), TEXT("Y")  },
+			{ TEXT("TestStruct"), TEXT("Nested"), TEXT("Vector"), TEXT("Z")  }
+			};
+		const TArray<FExpectedPropertyPath> Overlap = {
+			{ TEXT("Vector"), TEXT("X") },
+			{ TEXT("TestStruct") },
+		};
+
+		// 2. Convert to FConcertPropertyChain
+		FConcertPropertySelection RootPropertySelection;
+		FConcertPropertySelection NestedPropertySelection;
+		FConcertPropertySelection OverlapSelection;
+		auto TransformOp = [this](const FExpectedPropertyPath& Path)
+		{
+			const TOptional<FConcertPropertyChain> Result = FConcertPropertyChain::CreateFromPath(*UTestReflectionObject::StaticClass(), Path);
+			if (!Result)
+			{
+				AddError(TEXT("Failed to convert all paths"));
+			}
+			return Result.Get({});
+		};
+		Algo::Transform(RootProperties, RootPropertySelection.ReplicatedProperties, TransformOp);
+		Algo::Transform(NestedProperties, NestedPropertySelection.ReplicatedProperties, TransformOp);
+		Algo::Transform(Overlap, OverlapSelection.ReplicatedProperties, TransformOp);
+
 
-	// This is technically unnecessary because Memcmp would already have failed the test above but this helps in debugging what values caused the Memcmp to be different.
-	ValidateTestValues(*TestObject);
-	TestEqual(TEXT("Object.TestStruct.Nested.Vector was left unchanged"), TestObject->TestStruct.Nested.Vector, FVector{ 100.f, 200.f, 300.f });
-	return true;
+		// 3. Test permutations
+		TestFalse(TEXT("Root <> Nested"), RootPropertySelection.OverlapsWith(NestedPropertySelection));
+		TestFalse(TEXT("Nested <> Root"), NestedPropertySelection.OverlapsWith(RootPropertySelection));
+		
+		TestTrue(TEXT("Root <> Overlap"), RootPropertySelection.OverlapsWith(OverlapSelection));
+		TestTrue(TEXT("Nested <> Overlap"), NestedPropertySelection.OverlapsWith(OverlapSelection));
+		TestTrue(TEXT("Overlap <> Root"), OverlapSelection.OverlapsWith(RootPropertySelection));
+		TestTrue(TEXT("Overlap <> Nested"), OverlapSelection.OverlapsWith(NestedPropertySelection));
+		
+		return true;
+	}
 }
\ No newline at end of file
diff --git a/Engine/Plugins/Developer/Concert/ConcertSync/ConcertSyncTest/Source/ConcertSyncTest/Private/Replication/ReplicationAuthorityTests.cpp b/Engine/Plugins/Developer/Concert/ConcertSync/ConcertSyncTest/Source/ConcertSyncTest/Private/Replication/ReplicationAuthorityTests.cpp
new file mode 100644
--- /dev/null
+++ b/Engine/Plugins/Developer/Concert/ConcertSync/ConcertSyncTest/Source/ConcertSyncTest/Private/Replication/ReplicationAuthorityTests.cpp
@@ -0,0 +1,223 @@
+// Copyright Epic Games, Inc. All Rights Reserved.
+
+#include "Util/SendReceiveObjectTestBase.h"
+
+#include "Replication/IConcertClientReplicationManager.h"
+#include "TestReflectionObject.h"
+
+#include "Components/StaticMeshComponent.h"
+#include "Misc/AutomationTest.h"
+#include "Misc/Paths.h"
+#include "HAL/FileManager.h"
+#include "Util/SendReceiveGenericStreamTestBase.h"
+
+namespace UE::ConcertSyncTests::Replication::Authority
+{
+	/** Data sent from a client that does not have authority over objects is rejected. */
+	IMPLEMENT_CUSTOM_SIMPLE_AUTOMATION_TEST(FRejectUnauthorativeClientTest, FSendReceiveObjectTestBase, "Concert.Replication.Authority.RejectUnauthorativeClient", EAutomationTestFlags::EditorContext | EAutomationTestFlags::EngineFilter);
+	bool FRejectUnauthorativeClientTest::RunTest(const FString& Parameters)
+	{
+		// 1. Init
+		SetUpClientAndServer();
+		
+		// 2. Send data without authority > Reject
+		bool bHasServerReceivedData = false;
+		auto OnServerReceive = [this, &bHasServerReceivedData](const FConcertSessionContext& Context, const FConcertBatchReplicationEvent& Event) mutable
+		{
+			if (bHasServerReceivedData)
+			{
+				AddError(TEXT("Server was expected to receive data exactly once!"));
+			}
+			bHasServerReceivedData = true;
+		};
+		auto OnClientReceive = [this](const FConcertSessionContext& Context, const FConcertBatchReplicationEvent& Event) mutable
+		{
+			AddError(TEXT("Server sent data from non-authorative client to receiving client!"));
+		};
+		SimulateSenderToReceiver(OnServerReceive, OnClientReceive);
+		TestTrue(TEXT("Server received replication event from non-authorative client"), bHasServerReceivedData);
+		
+		return true;
+	}
+
+	/** Makes sure only one client can take authority over the same object properties at the same time and that other clients continue being able to take authority after authority has been released.  */
+	IMPLEMENT_CUSTOM_SIMPLE_AUTOMATION_TEST(FAuthorityIsMutallyExclusiveTest, FSendReceiveGenericTestBase, "Concert.Replication.Authority.AuthorityIsMutallyExclusive", EAutomationTestFlags::EditorContext | EAutomationTestFlags::EngineFilter);
+	bool FAuthorityIsMutallyExclusiveTest::RunTest(const FString& Parameters)
+	{
+		// 1. Init
+		const UObject* UnusedTestObject = GetDefault<UTestReflectionObject>();
+		const FSoftObjectPath TestObjectPath { UnusedTestObject };
+		SenderArgs = CreateHandshakeArgsFrom(*UnusedTestObject);
+		ReceiverArgs = SenderArgs;
+		
+		SetUpClientAndServer();
+
+		// 2.1 Senders takes authority over an object...
+		bool bSenderReceivedResponse = false;
+		ClientReplicationManager_Sender->TakeAuthorityOver({ TestObjectPath })
+			.Next([this, &bSenderReceivedResponse](const ConcertSyncClient::Replication::FAuthorityChangeResponse& Response) mutable
+			{
+				bSenderReceivedResponse = true;
+				TestEqual(TEXT("No rejection taking authority"), Response.RejectedObjects.Num(), 0); 
+			});
+		TickClient(Client_Sender);
+		TickServer();
+		TestTrue(TEXT("Sender received response to take authority"), bSenderReceivedResponse);
+
+		// 2.2 ... so does receiver fails taking authority over the same object
+		bool bReceiverReceivedResponse = false;
+		ClientReplicationManager_Receiver->TakeAuthorityOver({ TestObjectPath })
+			.Next([this, &bReceiverReceivedResponse](const ConcertSyncClient::Replication::FAuthorityChangeResponse& Response) mutable
+			{
+				bReceiverReceivedResponse = true;
+				TestEqual(TEXT("Rejected because Sender already has authority"), Response.RejectedObjects.Num(), 1); 
+			});
+		TickClient(Client_Receiver);
+		TickServer();
+		TestTrue(TEXT("Receiver received rejection response"), bReceiverReceivedResponse);
+
+		// 2.3 ... then the sender lets go of authority ...
+		bSenderReceivedResponse = false;
+		ClientReplicationManager_Sender->ReleaseAuthorityOf({ TestObjectPath })
+			.Next([this, &bSenderReceivedResponse](const ConcertSyncClient::Replication::FAuthorityChangeResponse& Response) mutable
+			{
+				bSenderReceivedResponse = true;
+				TestEqual(TEXT("No rejection releasing object"), Response.RejectedObjects.Num(), 0); 
+			});
+		TickClient(Client_Sender);
+		TickServer();
+		TestTrue(TEXT("Sender received response releasing object"), bSenderReceivedResponse);
+
+		// 2.4 ... and the receiver can take authority of the object
+		bReceiverReceivedResponse = false;
+		ClientReplicationManager_Receiver->TakeAuthorityOver({ TestObjectPath })
+			.Next([this, &bReceiverReceivedResponse](const ConcertSyncClient::Replication::FAuthorityChangeResponse& Response) mutable
+			{
+				bReceiverReceivedResponse = true;
+				TestEqual(TEXT("No rejection because the object should not be released"), Response.RejectedObjects.Num(), 0); 
+			});
+		TickClient(Client_Receiver);
+		TickServer();
+		TestTrue(TEXT("Receiver received response taking authority"), bReceiverReceivedResponse);
+		
+		return true;
+	}
+	
+	/** Clients cannot take authority over objects that they did not send */
+	IMPLEMENT_CUSTOM_SIMPLE_AUTOMATION_TEST(FCannotTakeAuthorityOverUnregisteredObjectsTest, FSendReceiveGenericTestBase, "Concert.Replication.Authority.CannotTakeAuthorityOverUnregisteredObjects", EAutomationTestFlags::EditorContext | EAutomationTestFlags::EngineFilter);
+	bool FCannotTakeAuthorityOverUnregisteredObjectsTest::RunTest(const FString& Parameters)
+	{
+		// FSendReceiveGenericTestBase is used here because it registers no sending streams by default
+		
+		// 1. Init
+		SetUpClientAndServer();
+
+		// 2. Attempt to take authority over object that was not in handshake
+		bool bReceivedResponse = false;
+		const FSoftObjectPath SomePath(GetDefault<UTestReflectionObject>());
+		ConcertSyncClient::Replication::FAuthorityChangeRequest Request;
+		Request.TakeAuthority.Add(SomePath, FConcertStreamArray{ .StreamIds = { FGuid::NewGuid() } });
+		// Detail: Cannot use IConcertClientReplicationManager::TakeAuthority util here because it builds the request based on what streams were registered
+		ClientReplicationManager_Sender->RequestAuthorityChange(Request)
+			.Next([this, &bReceivedResponse, SomePath](const ConcertSyncClient::Replication::FAuthorityChangeResponse& Response) mutable
+			{
+				bReceivedResponse = true;
+				TestTrue(TEXT("Cannot take authority over unregistered object"), Response.RejectedObjects.Contains(SomePath));
+				TestEqual(TEXT("Rejected exactly 1 object"), Response.RejectedObjects.Num(), 1); 
+			});
+
+		// 3. Authority request was rejected
+		SimulateSenderToReceiver();
+		TestTrue(TEXT("Authority request was answered"), bReceivedResponse);
+		
+		return true;
+	}
+
+	/** When a client leaves the replication session, the client releases all of the authority as well.  */
+	IMPLEMENT_CUSTOM_SIMPLE_AUTOMATION_TEST(FClientLeavingReplicationSessionLosesAuthority, FSendReceiveGenericTestBase, "Concert.Replication.Authority.ClientLeavingReplicationSessionLosesAuthority", EAutomationTestFlags::EditorContext | EAutomationTestFlags::EngineFilter);
+	bool FClientLeavingReplicationSessionLosesAuthority::RunTest(const FString& Parameters)
+	{
+		// 1. Init
+		const UObject* UnusedTestObject = GetDefault<UTestReflectionObject>();
+		const FSoftObjectPath TestObjectPath { UnusedTestObject };
+		SenderArgs = CreateHandshakeArgsFrom(*UnusedTestObject);
+		ReceiverArgs = SenderArgs;
+		
+		SetUpClientAndServer();
+
+		// 2.1 Senders takes authority over an object...
+		bool bSenderReceivedResponse = false;
+		ClientReplicationManager_Sender->TakeAuthorityOver({ TestObjectPath })
+			.Next([this, &bSenderReceivedResponse](const ConcertSyncClient::Replication::FAuthorityChangeResponse& Response) mutable
+			{
+				bSenderReceivedResponse = true;
+				TestEqual(TEXT("No rejection taking authority (sender)"), Response.RejectedObjects.Num(), 0); 
+			});
+		TickClient(Client_Sender);
+		TickServer();
+		TestTrue(TEXT("Sender received response to take authority"), bSenderReceivedResponse);
+
+		// 2.2 ... then leaves
+		ClientReplicationManager_Sender->LeaveReplicationSession();
+		TickClient(Client_Sender);
+		TickServer();
+		
+		// 2.3 which means the receiver is allowed to take authority
+		bool bReceiverReceivedResponse = false;
+		ClientReplicationManager_Receiver->TakeAuthorityOver({ TestObjectPath })
+			.Next([this, &bReceiverReceivedResponse](const ConcertSyncClient::Replication::FAuthorityChangeResponse& Response) mutable
+			{
+				bReceiverReceivedResponse = true;
+				TestEqual(TEXT("No rejection taking authority (receiver)"), Response.RejectedObjects.Num(), 0); 
+			});
+		TickClient(Client_Receiver);
+		TickServer();
+		TestTrue(TEXT("Receiver received response to take authority"), bReceiverReceivedResponse);
+		
+		return true;
+	}
+
+	/** When a client leaves the replication session, the client releases all of the authority as well.  */
+	IMPLEMENT_CUSTOM_SIMPLE_AUTOMATION_TEST(FClientLeavingConcertSessionLosesAuthority, FSendReceiveGenericTestBase, "Concert.Replication.Authority.ClientLeavingConcertSessionLosesAuthority", EAutomationTestFlags::EditorContext | EAutomationTestFlags::EngineFilter);
+	bool FClientLeavingConcertSessionLosesAuthority::RunTest(const FString& Parameters)
+	{
+		// 1. Init
+		const UObject* UnusedTestObject = GetDefault<UTestReflectionObject>();
+		const FSoftObjectPath TestObjectPath { UnusedTestObject };
+		SenderArgs = CreateHandshakeArgsFrom(*UnusedTestObject);
+		ReceiverArgs = SenderArgs;
+		
+		SetUpClientAndServer();
+
+		// 2.1 Senders takes authority over an object...
+		bool bSenderReceivedResponse = false;
+		ClientReplicationManager_Sender->TakeAuthorityOver({ TestObjectPath })
+			.Next([this, &bSenderReceivedResponse](const ConcertSyncClient::Replication::FAuthorityChangeResponse& Response) mutable
+			{
+				bSenderReceivedResponse = true;
+				TestEqual(TEXT("No rejection taking authority (sender)"), Response.RejectedObjects.Num(), 0); 
+			});
+		TickClient(Client_Sender);
+		TickServer();
+		TestTrue(TEXT("Sender received response to take authority"), bSenderReceivedResponse);
+
+		// 2.2 ... then leaves
+		Client_Sender->ClientSessionMock->Disconnect();
+		TickClient(Client_Sender);
+		TickServer();
+		
+		// 2.3 which means the receiver is allowed to take authority
+		bool bReceiverReceivedResponse = false;
+		ClientReplicationManager_Receiver->TakeAuthorityOver({ TestObjectPath })
+			.Next([this, &bReceiverReceivedResponse](const ConcertSyncClient::Replication::FAuthorityChangeResponse& Response) mutable
+			{
+				bReceiverReceivedResponse = true;
+				TestEqual(TEXT("No rejection taking authority (receiver)"), Response.RejectedObjects.Num(), 0); 
+			});
+		TickClient(Client_Receiver);
+		TickServer();
+		TestTrue(TEXT("Receiver received response to take authority"), bReceiverReceivedResponse);
+		
+		return true;
+	}
+}
diff --git a/Engine/Plugins/Developer/Concert/ConcertSync/ConcertSyncTest/Source/ConcertSyncTest/Private/Replication/ReplicationHandshakeTests.cpp b/Engine/Plugins/Developer/Concert/ConcertSync/ConcertSyncTest/Source/ConcertSyncTest/Private/Replication/ReplicationHandshakeTests.cpp
--- a/Engine/Plugins/Developer/Concert/ConcertSync/ConcertSyncTest/Source/ConcertSyncTest/Private/Replication/ReplicationHandshakeTests.cpp
+++ b/Engine/Plugins/Developer/Concert/ConcertSync/ConcertSyncTest/Source/ConcertSyncTest/Private/Replication/ReplicationHandshakeTests.cpp
@@ -16,12 +16,12 @@
 #include "Misc/Paths.h"
 #include "HAL/FileManager.h"
 
-namespace UE::ConcertSyncTests::Replication
+namespace UE::ConcertSyncTests::Replication::Handshake
 {
 	/**
 	 * Tests the handshake for joining and leaving a replication session.
 	 */
-	IMPLEMENT_CUSTOM_SIMPLE_AUTOMATION_TEST(FJoinHandshakeTest, FConcertClientServerCommunicationTest, "Concert.Replication.JoinHandshake", EAutomationTestFlags::EditorContext | EAutomationTestFlags::EngineFilter);
+	IMPLEMENT_CUSTOM_SIMPLE_AUTOMATION_TEST(FJoinHandshakeTest, FConcertClientServerCommunicationTest, "Concert.Replication.Handshake", EAutomationTestFlags::EditorContext | EAutomationTestFlags::EngineFilter);
 	bool FJoinHandshakeTest::RunTest(const FString& Parameters)
 	{
 		using namespace ConcertSyncClient::TestInterface;
@@ -43,7 +43,7 @@ namespace UE::ConcertSyncTests::Replication
 		FClientInfo& Client = ConnectClient();
 		const TSharedRef<IConcertClientReplicationBridge> BridgeMock = MakeShared<FConcertClientReplicationBridgeMock>();
 		const TSharedPtr<IConcertClientSession>& ClientSession = Client.ClientSessionMock;
-		const TSharedRef<IConcertClientReplicationManager> ClientReplicationManager = CreateClientReplicationManager(ClientSession.ToSharedRef(), &BridgeMock.Get());
+		const TSharedRef<IConcertClientReplicationManager> ClientReplicationManager_Primary = CreateClientReplicationManager(ClientSession.ToSharedRef(), &BridgeMock.Get());
 		FClientInfo& Client_Secondary = ConnectClient();
 		const TSharedRef<IConcertClientReplicationBridge> BridgeMock_Secondary = MakeShared<FConcertClientReplicationBridgeMock>();
 		const TSharedPtr<IConcertClientSession>& ClientSession_Secondary = Client_Secondary.ClientSessionMock;
@@ -69,68 +69,62 @@ namespace UE::ConcertSyncTests::Replication
 		// 2. Run
 
 		// 2.1 Invalid configurations
+		// These should cases never happen if using the editor tools.
+		// However malicious users or those with custom C++ logic can send whatever they want.
+		
 		// 2.1.1 Duplicate properties
 		FReplicationStreamDescription Invalid_DoublePropertyDescription = StreamDescription;
 		Invalid_DoublePropertyDescription.BaseDescription.ReplicationMap.ReplicatedObjects[PathToSomeActorComponent].PropertySelection.ReplicatedProperties.Add(ForcedLodModelProperty);
-		ClientReplicationManager->JoinReplicationSession({ ClientDescription, { Invalid_DoublePropertyDescription } })
+		ClientReplicationManager_Primary->JoinReplicationSession({ ClientDescription, { Invalid_DoublePropertyDescription } })
 			.Next([&](const FJoinReplicatedSessionResult& Result)
 			{
 				TestTrue(TEXT("Cannot contain same properties twice"), Result.ErrorCode == EJoinReplicationErrorCode::DuplicateProperty);
 			});
 		// 2.1.2 Duplicate stream identifier
-		ClientReplicationManager->JoinReplicationSession({ ClientDescription, { StreamDescription, StreamDescription } })
+		ClientReplicationManager_Primary->JoinReplicationSession({ ClientDescription, { StreamDescription, StreamDescription } })
 			.Next([&](const FJoinReplicatedSessionResult& Result)
 			{
-				TestTrue(TEXT("Cannot contain stream ID twice"), Result.ErrorCode == EJoinReplicationErrorCode::ConflictingStreamId);
+				TestTrue(TEXT("Cannot contain stream ID twice"), Result.ErrorCode == EJoinReplicationErrorCode::DuplicateStreamId);
 			});
-		// 2.1.3 Miss class path
-		ClientReplicationManager->JoinReplicationSession({ ClientDescription, { InvalidClassStreamDescription } })
+		// 2.1.3 Missing class path
+		ClientReplicationManager_Primary->JoinReplicationSession({ ClientDescription, { InvalidClassStreamDescription } })
 			.Next([&](const FJoinReplicatedSessionResult& Result)
 			{
 				TestTrue(TEXT("Cannot contain null classes"), Result.ErrorCode == EJoinReplicationErrorCode::InvalidClass);
 			});
 
-		// 2.2 Valid join
-		ClientReplicationManager->JoinReplicationSession({ ClientDescription, { StreamDescription } })
+
+		// 2.2 Real workflow cases
+		// 2.2.1 Valid join
+		ClientReplicationManager_Primary->JoinReplicationSession({ ClientDescription, { StreamDescription } })
 			.Next([&](const FJoinReplicatedSessionResult& Result)
 			{
 				TestTrue(TEXT("Join with valid args"), Result.ErrorCode == EJoinReplicationErrorCode::Success);
 			});
-		// 2.3 No joining twice
+		// 2.2.2 No joining twice
 		AddExpectedError(TEXT("JoinReplicationSession requested while already in a session"), EAutomationExpectedErrorFlags::Contains); // Not pretty, but otherwise this test fails due to logged warning
-		ClientReplicationManager->JoinReplicationSession({ ClientDescription, { StreamDescription } })
+		ClientReplicationManager_Primary->JoinReplicationSession({ ClientDescription, { StreamDescription } })
 			.Next([&](const FJoinReplicatedSessionResult& Result)
 			{
 				TestTrue(TEXT("Cannot join twice"), Result.ErrorCode == EJoinReplicationErrorCode::AlreadyInSession);
 			});
-		// 2.4 Rejoining
-		ClientReplicationManager->LeaveReplicationSession();
-		ClientReplicationManager->JoinReplicationSession({ ClientDescription, { StreamDescription } })
+		// 2.2.3 Rejoining
+		ClientReplicationManager_Primary->LeaveReplicationSession();
+		ClientReplicationManager_Primary->JoinReplicationSession({ ClientDescription, { StreamDescription } })
 			.Next([&](const FJoinReplicatedSessionResult& Result)
 			{
 				TestTrue(TEXT("Re-join session"), Result.ErrorCode == EJoinReplicationErrorCode::Success);
 			});
-
 		
-		// 2.5 Second client is not allowed replicate the same properties as the first client
+		// 2.2.4 Two clients may join with the same stream identifier and properties. 
+		// For context: UX-wise it should be as easy as possible for clients to join.
+		// This is why we allow clients to join with overlapping properties and stream IDs, which used to be rejected by the handshake.
+		// Now, they must take authority over objects before sending. See authority tests.
 		FReplicationStreamDescription DuplicateProperties = StreamDescription;
-		DuplicateProperties.BaseDescription.Identifier = FGuid::NewGuid();
-		ClientReplicationManager_Secondary->JoinReplicationSession({ ClientDescription, { DuplicateProperties } })
-			.Next([&](const FJoinReplicatedSessionResult& Result)
-			{
-				TestTrue(TEXT("Clients cannot overlap properties"), Result.ErrorCode == EJoinReplicationErrorCode::ConflictingAuthority);
-			});
-		// 2.6 Second client is allowed to replicate different properties on the same object
-		FReplicationStreamDescription NonOverlappingProperties;
-		NonOverlappingProperties.BaseDescription.Identifier = FGuid::NewGuid();
-		FReplicatedObjectInfo StaticMeshComponentInfo_Secondary;
-		StaticMeshComponentInfo_Secondary.ClassPath = UStaticMeshComponent::StaticClass();
-		StaticMeshComponentInfo_Secondary.PropertySelection.ReplicatedProperties.Add(MinLODProperty);
-		NonOverlappingProperties.BaseDescription.ReplicationMap.ReplicatedObjects.Add(PathToSomeActorComponent, StaticMeshComponentInfo_Secondary);
-		ClientReplicationManager_Secondary->JoinReplicationSession({ ClientDescription, { NonOverlappingProperties } })
+		ClientReplicationManager_Secondary->JoinReplicationSession({ ClientDescription, { StreamDescription } })
 			.Next([&](const FJoinReplicatedSessionResult& Result)
 			{
-				TestTrue(TEXT("Two clients can replicate differing properties on the same object"), Result.ErrorCode == EJoinReplicationErrorCode::Success);
+				TestTrue(TEXT("2nd client can overlap stream identifier and properties of 1st client"), Result.ErrorCode == EJoinReplicationErrorCode::Success);
 			});
 		
 		return true;
diff --git a/Engine/Plugins/Developer/Concert/ConcertSync/ConcertSyncTest/Source/ConcertSyncTest/Private/Replication/ReplicationSendReceiveTests.cpp b/Engine/Plugins/Developer/Concert/ConcertSync/ConcertSyncTest/Source/ConcertSyncTest/Private/Replication/ReplicationSendReceiveTests.cpp
--- a/Engine/Plugins/Developer/Concert/ConcertSync/ConcertSyncTest/Source/ConcertSyncTest/Private/Replication/ReplicationSendReceiveTests.cpp
+++ b/Engine/Plugins/Developer/Concert/ConcertSync/ConcertSyncTest/Source/ConcertSyncTest/Private/Replication/ReplicationSendReceiveTests.cpp
@@ -16,46 +16,108 @@
 #include "Misc/AutomationTest.h"
 #include "Misc/Paths.h"
 #include "HAL/FileManager.h"
+#include "UObject/Class.h"
 #include "UObject/UObjectGlobals.h"
 #include "UObject/Package.h"
 
-namespace UE::ConcertSyncTests::Replication
+namespace UE::ConcertSyncTests::Replication::SendReceiveFlow
 {
-	/**
-	 * Tests replicating data from sender client > server > receiver client.
-	 */
-	IMPLEMENT_CUSTOM_SIMPLE_AUTOMATION_TEST(FSendReceiveFlowTests, FSendReceiveObjectTestBase, "Concert.Replication.SendReceiveFlow", EAutomationTestFlags::EditorContext | EAutomationTestFlags::EngineFilter);
-	bool FSendReceiveFlowTests::RunTest(const FString& Parameters)
+	static bool SharedRunTest(FSendReceiveObjectTestBase& Test)
 	{
 		// 1. Init
-		SetUpClientAndServer();
+		Test.SetUpClientAndServer();
+		
+		Test.ClientReplicationManager_Sender->TakeAuthorityOver({ Test.TestObject })
+			.Next([&Test](ConcertSyncClient::Replication::FAuthorityChangeResponse&& Response)
+			{
+				if (!Response.RejectedObjects.IsEmpty())
+				{
+					Test.AddError(TEXT("Failed to take authority"));
+				}
+			});
 		
 		// 2. Send data
 		bool bHasServerReceivedData = false;
 		bool bHasClientReceivedData = false;
-		auto OnServerReceive = [this, &bHasServerReceivedData](const FConcertSessionContext& Context, const FConcertBatchReplicationEvent& Event) mutable
+		auto OnServerReceive = [&Test, &bHasServerReceivedData](const FConcertSessionContext& Context, const FConcertBatchReplicationEvent& Event) mutable
 		{
 			if (bHasServerReceivedData)
 			{
-				AddError(TEXT("Server was expected to receive data exactly once!"));
+				Test.AddError(TEXT("Server was expected to receive data exactly once!"));
 			}
 			bHasServerReceivedData = true;
 		};
-		auto OnClientReceive = [this, &bHasClientReceivedData](const FConcertSessionContext& Context, const FConcertBatchReplicationEvent& Event) mutable
+		auto OnClientReceive = [&Test, &bHasClientReceivedData](const FConcertSessionContext& Context, const FConcertBatchReplicationEvent& Event) mutable
 		{
 			if (bHasClientReceivedData)
 			{
-				AddError(TEXT("Client 2 was expected to receive data exactly once!"));
+				Test.AddError(TEXT("Client 2 was expected to receive data exactly once!"));
 			}
 			bHasClientReceivedData = true;
 		};
-		SimulateSenderToReceiver(OnServerReceive, OnClientReceive);
+		Test.SimulateSendObjectToReceiver(OnServerReceive, OnClientReceive);
 
 		// 3. Test
-		TestTrue(TEXT("Server received replication event"), bHasServerReceivedData);
-		TestTrue(TEXT("Client 2 received replication event"), bHasClientReceivedData);
-		TestEqualTestValues(*TestObject, *this);
+		Test.TestTrue(TEXT("Server received replication event"), bHasServerReceivedData);
+		Test.TestTrue(TEXT("Client 2 received replication event"), bHasClientReceivedData);
+		Test.TestEqualTestValues(*Test.TestObject, Test);
 		
 		return true;
 	}
+	
+	/**
+	 * Tests replicating data from sender client > server > receiver client.
+	 */
+	IMPLEMENT_CUSTOM_SIMPLE_AUTOMATION_TEST(FSendReceiveFlowTests, FSendReceiveObjectTestBase, "Concert.Replication.SendReceive.SingleStream", EAutomationTestFlags::EditorContext | EAutomationTestFlags::EngineFilter);
+	bool FSendReceiveFlowTests::RunTest(const FString& Parameters)
+	{
+		return SharedRunTest(*this);
+	}
+
+	/** Test which still sends TestObject but does so with two separate streams: one for the float and the other for the vector property. */
+	class FSplitSendReceiveObjectTest : public FSendReceiveObjectTestBase
+	{
+	public:
+		FSplitSendReceiveObjectTest(const FString& InName, const bool bInComplexTask)
+        	: FSendReceiveObjectTestBase(InName, bInComplexTask)
+        {}
+
+	protected:
+
+		FGuid FloatStreamId = FGuid::NewGuid();
+		FGuid VectorStreamId = FGuid::NewGuid();
+		
+		//~ Begin FSendReceiveTestBase Interface
+		virtual ConcertSyncClient::Replication::FJoinReplicatedSessionArgs CreateSenderArgs() override
+		{
+			ConcertSyncClient::Replication::FJoinReplicatedSessionArgs SenderJoinArgs;
+			const UClass& TestClass = *TestObject->GetClass();
+			
+			FReplicatedObjectInfo FloatProperties { &TestClass };
+			FloatProperties.PropertySelection.ReplicatedProperties.Add(*FConcertPropertyChain::CreateFromPath(TestClass, { GET_MEMBER_NAME_CHECKED(UTestReflectionObject, Float) }));
+			FReplicatedObjectInfo VectorProperties { &TestClass };
+			VectorProperties.PropertySelection.ReplicatedProperties.Add(*FConcertPropertyChain::CreateFromPath(TestClass, { GET_MEMBER_NAME_CHECKED(UTestReflectionObject, Vector) }));
+
+			FReplicationStreamDescription FloatStream;
+			FloatStream.BaseDescription.Identifier = FloatStreamId;
+			FloatStream.BaseDescription.ReplicationMap.ReplicatedObjects.Add(TestObject, FloatProperties);
+			FReplicationStreamDescription VectorStream;
+			VectorStream.BaseDescription.Identifier = VectorStreamId;
+			VectorStream.BaseDescription.ReplicationMap.ReplicatedObjects.Add(TestObject, VectorProperties);
+			
+			SenderJoinArgs.Streams.Add(FloatStream);
+			SenderJoinArgs.Streams.Add(VectorStream);
+			return SenderJoinArgs;
+		}
+
+		virtual TSet<FGuid> GetSenderStreamIds() const override { return { FloatStreamId, VectorStreamId }; }
+		//~ End FSendReceiveTestBase Interface
+	};
+
+	/** Test which still sends TestObject but does so with two separate streams: one for the float and the other for the vector property. */
+	IMPLEMENT_CUSTOM_SIMPLE_AUTOMATION_TEST(FSendReceiveMultiStreamSameObjectTests, FSplitSendReceiveObjectTest, "Concert.Replication.SendReceive.MultipleStreamsForSameObject", EAutomationTestFlags::EditorContext | EAutomationTestFlags::EngineFilter);
+	bool FSendReceiveMultiStreamSameObjectTests::RunTest(const FString& Parameters)
+	{
+		return SharedRunTest(*this);
+	}
 }
diff --git a/Engine/Plugins/Developer/Concert/ConcertSync/ConcertSyncTest/Source/ConcertSyncTest/Private/Replication/Util/SendReceiveGenericStreamTestBase.h b/Engine/Plugins/Developer/Concert/ConcertSync/ConcertSyncTest/Source/ConcertSyncTest/Private/Replication/Util/SendReceiveGenericStreamTestBase.h
new file mode 100644
--- /dev/null
+++ b/Engine/Plugins/Developer/Concert/ConcertSync/ConcertSyncTest/Source/ConcertSyncTest/Private/Replication/Util/SendReceiveGenericStreamTestBase.h
@@ -0,0 +1,41 @@
+// Copyright Epic Games, Inc. All Rights Reserved.
+
+#pragma once
+
+#include "SendReceiveTestBase.h"
+
+struct FConcertBatchReplicationEvent;
+class IConcertClientReplicationManager;
+class UTestReflectionObject;
+
+namespace UE::ConcertSyncServer::Replication
+{
+	class IConcertServerReplicationManager;
+}
+
+namespace UE::ConcertSyncTests::Replication
+{
+	class FConcertClientReplicationBridgeMock;
+	
+	/**
+	 * Simple test exposing SenderArgs and ReceiverArgs and implements CreateSenderArgs and CreateReceiverArgs to return them, respectively.
+	 */
+	class FSendReceiveGenericTestBase : public FSendReceiveTestBase
+	{
+	public:
+
+		FSendReceiveGenericTestBase(const FString& InName, const bool bInComplexTask)
+			: FSendReceiveTestBase(InName, bInComplexTask)
+		{}
+
+	protected:
+		
+		ConcertSyncClient::Replication::FJoinReplicatedSessionArgs SenderArgs;
+		ConcertSyncClient::Replication::FJoinReplicatedSessionArgs ReceiverArgs;
+
+		//~ Begin FSendReceiveTestBase Interface
+		virtual ConcertSyncClient::Replication::FJoinReplicatedSessionArgs CreateSenderArgs() override { return SenderArgs; }
+		virtual ConcertSyncClient::Replication::FJoinReplicatedSessionArgs CreateReceiverArgs() override { return ReceiverArgs; }
+		//~ End FSendReceiveTestBase Interface
+	};
+}
diff --git a/Engine/Plugins/Developer/Concert/ConcertSync/ConcertSyncTest/Source/ConcertSyncTest/Private/Replication/Util/SendReceiveObjectTestBase.cpp b/Engine/Plugins/Developer/Concert/ConcertSync/ConcertSyncTest/Source/ConcertSyncTest/Private/Replication/Util/SendReceiveObjectTestBase.cpp
--- a/Engine/Plugins/Developer/Concert/ConcertSync/ConcertSyncTest/Source/ConcertSyncTest/Private/Replication/Util/SendReceiveObjectTestBase.cpp
+++ b/Engine/Plugins/Developer/Concert/ConcertSync/ConcertSyncTest/Source/ConcertSyncTest/Private/Replication/Util/SendReceiveObjectTestBase.cpp
@@ -23,20 +23,7 @@ namespace UE::ConcertSyncTests::Replication
 {
 	ConcertSyncClient::Replication::FJoinReplicatedSessionArgs FSendReceiveObjectTestBase::CreateSenderArgs()
 	{
-		ConcertSyncClient::Replication::FJoinReplicatedSessionArgs SenderJoinArgs;
-		
-		FReplicatedObjectInfo AllProperties { UTestReflectionObject::StaticClass() };
-		ConcertSyncCore::PropertyChain::ForEachReplicatableConcertProperty(*UTestReflectionObject::StaticClass(), [&AllProperties](FConcertPropertyChain&& Chain)
-		{
-			AllProperties.PropertySelection.ReplicatedProperties.Emplace(MoveTemp(Chain));
-			return EBreakBehavior::Continue;
-		});
-
-		FReplicationStreamDescription SendingStream;
-		SendingStream.BaseDescription.Identifier = SenderStreamId;
-		SendingStream.BaseDescription.ReplicationMap.ReplicatedObjects.Add(TestObject, AllProperties);
-		SenderJoinArgs.Streams.Add(SendingStream);
-		return SenderJoinArgs;
+		return CreateHandshakeArgsFrom(*TestObject, SenderStreamId);
 	}
 
 	ConcertSyncClient::Replication::FJoinReplicatedSessionArgs FSendReceiveObjectTestBase::CreateReceiverArgs()
@@ -58,28 +45,34 @@ namespace UE::ConcertSyncTests::Replication
 		BridgeMock_Receiver->InjectAvailableObject(*TestObject);
 	}
 	
-	void FSendReceiveObjectTestBase::SimulateSenderToReceiver(
+	void FSendReceiveObjectTestBase::SimulateSendObjectToReceiver(
 		TFunctionRef<FReceiveReplicationEventSignature> OnServerReceive,
 		TFunctionRef<FReceiveReplicationEventSignature> OnReceiverClientReceive)
 	{
 		auto TestReplicationData_Server = [this, OnServerReceive](const FConcertSessionContext& Context, const FConcertBatchReplicationEvent& Event)
 		{
-			TestEqual(TEXT("Server received 1 stream"), Event.Streams.Num(), 1);
-			TestEqual(TEXT("Server received 1 object"), Event.Streams.IsEmpty() ? 0 : Event.Streams[0].ReplicatedObjects.Num() , 1);
-			TestEqual(TEXT("Server received from correct stream"), Event.Streams.IsEmpty() ? FGuid{} : Event.Streams[0].StreamId, SenderStreamId);
-			const FSoftObjectPath ObjectPath = Event.Streams.IsEmpty() || Event.Streams[0].ReplicatedObjects.IsEmpty() ? FSoftObjectPath{} : Event.Streams[0].ReplicatedObjects[0].ReplicatedObject;
-			TestEqual(TEXT("Server's received object has correct path"), ObjectPath, FSoftObjectPath(TestObject));
-
+			const TSet<FGuid> SenderStreamIds = GetSenderStreamIds();
+			TestEqual(TEXT("Server received right number of streams"), Event.Streams.Num(), SenderStreamIds.Num());
+			for (int32 i = 0; i < Event.Streams.Num(); ++i)
+			{
+				TestEqual(TEXT("Server received 1 object"), Event.Streams[i].ReplicatedObjects.Num() , 1);
+				TestTrue(TEXT("Server received from correct stream"), SenderStreamIds.Contains(Event.Streams[i].StreamId));
+				const FSoftObjectPath ObjectPath = Event.Streams[i].ReplicatedObjects.IsEmpty() ? FSoftObjectPath{} : Event.Streams[i].ReplicatedObjects[0].ReplicatedObject;
+				TestEqual(TEXT("Server's received object has correct path"), ObjectPath, FSoftObjectPath(TestObject));
+			}
 			OnServerReceive(Context, Event);
 		};
 		auto TestReplicationData_Client_Receiver = [this, OnReceiverClientReceive](const FConcertSessionContext& Context, const FConcertBatchReplicationEvent& Event)
 		{
-			TestEqual(TEXT("Client 2 received 1 stream"), Event.Streams.Num(), 1);
-			TestEqual(TEXT("Client 2 received 1 object"), Event.Streams.IsEmpty() ? 0 : Event.Streams[0].ReplicatedObjects.Num() , 1);
-			TestEqual(TEXT("Client 2 received from correct stream"), Event.Streams.IsEmpty() ? FGuid{} : Event.Streams[0].StreamId, SenderStreamId);
-			const FSoftObjectPath ObjectPath = Event.Streams.IsEmpty() || Event.Streams[0].ReplicatedObjects.IsEmpty() ? FSoftObjectPath{} : Event.Streams[0].ReplicatedObjects[0].ReplicatedObject;
-			TestEqual(TEXT("Client 2's received object has correct path"), ObjectPath, FSoftObjectPath(TestObject));
-
+			const TSet<FGuid> SenderStreamIds = GetSenderStreamIds();
+			TestEqual(TEXT("Client 2 received right number of streams"), Event.Streams.Num(), SenderStreamIds.Num());
+			for (int32 i = 0; i < Event.Streams.Num(); ++i)
+			{
+				TestEqual(TEXT("Client 2 received 1 object"), Event.Streams[i].ReplicatedObjects.Num() , 1);
+				TestTrue(TEXT("Client 2 received from correct stream"), SenderStreamIds.Contains(Event.Streams[i].StreamId));
+				const FSoftObjectPath ObjectPath = Event.Streams[i].ReplicatedObjects.IsEmpty() ? FSoftObjectPath{} : Event.Streams[i].ReplicatedObjects[0].ReplicatedObject;
+				TestEqual(TEXT("Client 2's received object has correct path"), ObjectPath, FSoftObjectPath(TestObject));
+			}
 			OnReceiverClientReceive(Context, Event);
 		};
 		const FDelegateHandle ServerHandle = ServerSession->RegisterCustomEventHandler<FConcertBatchReplicationEvent>(TestReplicationData_Server);
@@ -101,6 +94,8 @@ namespace UE::ConcertSyncTests::Replication
 
 		ServerSession->UnregisterCustomEventHandler<FConcertBatchReplicationEvent>(ServerHandle);
 		ServerSession->UnregisterCustomEventHandler<FConcertBatchReplicationEvent>(ClientHandle);
+		
+		// No call to Super because we're completely overriding the behavior.
 	}
 
 	void FSendReceiveObjectTestBase::SetTestValues(UTestReflectionObject& Object)
diff --git a/Engine/Plugins/Developer/Concert/ConcertSync/ConcertSyncTest/Source/ConcertSyncTest/Private/Replication/Util/SendReceiveObjectTestBase.h b/Engine/Plugins/Developer/Concert/ConcertSync/ConcertSyncTest/Source/ConcertSyncTest/Private/Replication/Util/SendReceiveObjectTestBase.h
--- a/Engine/Plugins/Developer/Concert/ConcertSync/ConcertSyncTest/Source/ConcertSyncTest/Private/Replication/Util/SendReceiveObjectTestBase.h
+++ b/Engine/Plugins/Developer/Concert/ConcertSync/ConcertSyncTest/Source/ConcertSyncTest/Private/Replication/Util/SendReceiveObjectTestBase.h
@@ -28,28 +28,30 @@ namespace UE::ConcertSyncTests::Replication
 		FSendReceiveObjectTestBase(const FString& InName, const bool bInComplexTask)
 			: FSendReceiveTestBase(InName, bInComplexTask)
 		{}
-
-	protected:
-		
-		using FReceiveReplicationEventSignature = void(const FConcertSessionContext& Context, const FConcertBatchReplicationEvent& Event);
 		
+		/**
+		 * The object the test will be run on.
+		 * 
+		 * Note: this pointer is knowingly left dangling after the test completes...
+		 * there really is not need to override CleanUpTest because we assign TestObject first thing in SetUpClientAndServer.
+		 */
 		UTestReflectionObject* TestObject = nullptr;
 
 		//~ Begin FSendReceiveTestBase Interface
 		virtual ConcertSyncClient::Replication::FJoinReplicatedSessionArgs CreateSenderArgs() override;
 		virtual ConcertSyncClient::Replication::FJoinReplicatedSessionArgs CreateReceiverArgs() override;
 		virtual void SetUpClientAndServer() override;
+		//~ End FSendReceiveTestBase Interface
 		
 		/**
 		 * bHasServerReceivedData and bHasClientReceivedData are reset to false prior to sending.
 		 * Sets test values on TestObject and sends it to the receiver.
 		 * If the data arrived, bHasServerReceivedData and bHasClientReceivedData are true.
 		 */
-		void SimulateSenderToReceiver(
+		void SimulateSendObjectToReceiver(
 			TFunctionRef<FReceiveReplicationEventSignature> OnServerReceive = [](auto&, auto&){},
 			TFunctionRef<FReceiveReplicationEventSignature> OnReceiverClientReceive = [](auto&, auto&){}
 			);
-		//~ End FSendReceiveTestBase Interface
 
 		void SetTestValues(UTestReflectionObject& Object);
 		void SetDifferentValues(UTestReflectionObject& Object);
@@ -59,6 +61,9 @@ namespace UE::ConcertSyncTests::Replication
 		
 		const float SentFloat = 42.f;
 		const FVector SentVector = { 21.f, 84.f, -1.f };
+		
 		const FGuid SenderStreamId = FGuid::NewGuid();
+
+		virtual TSet<FGuid> GetSenderStreamIds() const { return { SenderStreamId }; } 
 	};
 }
diff --git a/Engine/Plugins/Developer/Concert/ConcertSync/ConcertSyncTest/Source/ConcertSyncTest/Private/Replication/Util/SendReceiveTestBase.cpp b/Engine/Plugins/Developer/Concert/ConcertSync/ConcertSyncTest/Source/ConcertSyncTest/Private/Replication/Util/SendReceiveTestBase.cpp
--- a/Engine/Plugins/Developer/Concert/ConcertSync/ConcertSyncTest/Source/ConcertSyncTest/Private/Replication/Util/SendReceiveTestBase.cpp
+++ b/Engine/Plugins/Developer/Concert/ConcertSync/ConcertSyncTest/Source/ConcertSyncTest/Private/Replication/Util/SendReceiveTestBase.cpp
@@ -12,6 +12,7 @@
 #include "Misc/AutomationTest.h"
 #include "Misc/Paths.h"
 #include "HAL/FileManager.h"
+#include "Replication/PropertyChainUtils.h"
 #include "UObject/UObjectGlobals.h"
 #include "UObject/Package.h"
 
@@ -21,6 +22,24 @@ namespace UE::ConcertSyncTests::Replication
 			: FConcertClientServerCommunicationTest(InName, bInComplexTask)
 	{}
 
+	ConcertSyncClient::Replication::FJoinReplicatedSessionArgs FSendReceiveTestBase::CreateHandshakeArgsFrom(const UObject& Object, const FGuid& SenderStreamId)
+	{
+		ConcertSyncClient::Replication::FJoinReplicatedSessionArgs SenderJoinArgs;
+		
+		FReplicatedObjectInfo AllProperties { Object.GetClass() };
+		ConcertSyncCore::PropertyChain::ForEachReplicatableConcertProperty(*Object.GetClass(), [&AllProperties](FConcertPropertyChain&& Chain)
+		{
+			AllProperties.PropertySelection.ReplicatedProperties.Emplace(MoveTemp(Chain));
+			return EBreakBehavior::Continue;
+		});
+
+		FReplicationStreamDescription SendingStream;
+		SendingStream.BaseDescription.Identifier = SenderStreamId;
+		SendingStream.BaseDescription.ReplicationMap.ReplicatedObjects.Add(&Object, AllProperties);
+		SenderJoinArgs.Streams.Add(SendingStream);
+		return SenderJoinArgs;
+	}
+
 	void FSendReceiveTestBase::SetUpClientAndServer()
 	{
 		// Server
@@ -62,14 +81,31 @@ namespace UE::ConcertSyncTests::Replication
 		Client_Sender->ClientSessionMock->RegisterCustomEventHandler<FConcertBatchReplicationEvent>(TestReplicationData_Client_Sender);
 	}
 	
-	void FSendReceiveTestBase::SimulateSenderToReceiver()
+	void FSendReceiveTestBase::SimulateSenderToReceiver(
+		TFunctionRef<FReceiveReplicationEventSignature> OnServerReceive,
+		TFunctionRef<FReceiveReplicationEventSignature> OnReceiverClientReceive
+		)
 	{
+		auto TestReplicationData_Server = [this, OnServerReceive](const FConcertSessionContext& Context, const FConcertBatchReplicationEvent& Event)
+		{
+			OnServerReceive(Context, Event);
+		};
+		auto TestReplicationData_Client_Receiver = [this, OnReceiverClientReceive](const FConcertSessionContext& Context, const FConcertBatchReplicationEvent& Event)
+		{
+			OnReceiverClientReceive(Context, Event);
+		};
+		const FDelegateHandle ServerHandle = ServerSession->RegisterCustomEventHandler<FConcertBatchReplicationEvent>(TestReplicationData_Server);
+		const FDelegateHandle ClientHandle = Client_Receiver->ClientSessionMock->RegisterCustomEventHandler<FConcertBatchReplicationEvent>(TestReplicationData_Client_Receiver);
+
 		// 1. Sender > Server
 		TickClient(Client_Sender);
 		// 2. Forward from server to receiver
 		TickServer();
 		// 3. Receive from server
 		TickClient(Client_Receiver);
+		
+		ServerSession->UnregisterCustomEventHandler<FConcertBatchReplicationEvent>(ServerHandle);
+		ServerSession->UnregisterCustomEventHandler<FConcertBatchReplicationEvent>(ClientHandle);
 	}
 	
 	void FSendReceiveTestBase::TickClient(FClientInfo* Client)
diff --git a/Engine/Plugins/Developer/Concert/ConcertSync/ConcertSyncTest/Source/ConcertSyncTest/Private/Replication/Util/SendReceiveTestBase.h b/Engine/Plugins/Developer/Concert/ConcertSync/ConcertSyncTest/Source/ConcertSyncTest/Private/Replication/Util/SendReceiveTestBase.h
--- a/Engine/Plugins/Developer/Concert/ConcertSync/ConcertSyncTest/Source/ConcertSyncTest/Private/Replication/Util/SendReceiveTestBase.h
+++ b/Engine/Plugins/Developer/Concert/ConcertSync/ConcertSyncTest/Source/ConcertSyncTest/Private/Replication/Util/SendReceiveTestBase.h
@@ -27,8 +27,10 @@ namespace UE::ConcertSyncTests::Replication
 	public:
 	
 		FSendReceiveTestBase(const FString& InName, const bool bInComplexTask);
-
-	protected:
+		
+		static ConcertSyncClient::Replication::FJoinReplicatedSessionArgs CreateHandshakeArgsFrom(const UObject& Object, const FGuid& SenderStreamId = FGuid::NewGuid());
+		
+		using FReceiveReplicationEventSignature = void(const FConcertSessionContext& Context, const FConcertBatchReplicationEvent& Event);
 		
 		float FakeDeltaTime = 1.f / 60.f;
 		
@@ -47,7 +49,10 @@ namespace UE::ConcertSyncTests::Replication
 		virtual ConcertSyncClient::Replication::FJoinReplicatedSessionArgs CreateReceiverArgs() = 0;
 
 		virtual void SetUpClientAndServer();
-		void SimulateSenderToReceiver();
+		virtual void SimulateSenderToReceiver(
+			TFunctionRef<FReceiveReplicationEventSignature> OnServerReceive = [](auto&, auto&){},
+			TFunctionRef<FReceiveReplicationEventSignature> OnReceiverClientReceive = [](auto&, auto&){}
+			);
 		
 		void TickClient(FClientInfo* Client);
 		void TickServer();
diff --git a/Engine/Plugins/Developer/Concert/ConcertSync/ConcertSyncTest/Source/ConcertSyncTest/Private/Util/ClientServerCommunicationTest.h b/Engine/Plugins/Developer/Concert/ConcertSync/ConcertSyncTest/Source/ConcertSyncTest/Private/Util/ClientServerCommunicationTest.h
--- a/Engine/Plugins/Developer/Concert/ConcertSync/ConcertSyncTest/Source/ConcertSyncTest/Private/Util/ClientServerCommunicationTest.h
+++ b/Engine/Plugins/Developer/Concert/ConcertSync/ConcertSyncTest/Source/ConcertSyncTest/Private/Util/ClientServerCommunicationTest.h
@@ -151,6 +151,11 @@ namespace UE::ConcertSyncTests
 	class FConcertServerSessionMock : public FConcertServerSessionBaseMock
 	{
 	public:
+
+		FConcertServerSessionMock()
+		{
+			RegisterCustomEventHandler<FConcertSession_LeaveSessionEvent>(this, &FConcertServerSessionMock::HandleClientLeaveEvent);
+		}
 		
 		virtual void InternalSendCustomEvent(const UScriptStruct* EventType, const void* EventData, const TArray<FGuid>& TargetEndpointIds, EConcertMessageFlags, TOptional<FConcertSequencedCustomEvent> InSequenceId={}) override
 		{
@@ -304,6 +309,17 @@ namespace UE::ConcertSyncTests
 
 		/** Connected clients sessions. */
 		TArray<FConcertClientSessionBaseMock*> ClientSessions;
+		
+		void HandleClientLeaveEvent(const FConcertSessionContext& Context, const FConcertSession_LeaveSessionEvent& Event)
+		{
+			const FGuid& ClientEndpointId = Context.SourceEndpointId;
+			ClientEndpoints.RemoveSingle(ClientEndpointId);
+			
+			const int32 Index = ClientSessions.IndexOfByPredicate([&ClientEndpointId](FConcertClientSessionBaseMock* SessionMock){ return SessionMock->GetSessionClientEndpointId() == ClientEndpointId; });
+			check(ClientSessions.IsValidIndex(Index));
+			
+			ConnectionChanged.Broadcast(*this, EConcertClientStatus::Disconnected, { ClientEndpointId });
+		}
 	};
 
 	/** Specializes the base concert client session to act as a fake client session. */
@@ -388,7 +404,17 @@ namespace UE::ConcertSyncTests
 		{
 			return FGuid(0, 0, 0, 0);
 		}
-		
+
+		virtual void Disconnect() override
+		{
+			if (ensure(bIsConnected))
+			{
+				bIsConnected = false;
+				FConcertSession_LeaveSessionEvent LeaveSessionEvent;
+				LeaveSessionEvent.SessionServerEndpointId = SessionInfo.ServerEndpointId;
+				ServerMock.DispatchEvent(EndpointId, FConcertSession_LeaveSessionEvent::StaticStruct(), &LeaveSessionEvent, { SessionInfo.ServerEndpointId }, EConcertMessageFlags::ReliableOrdered, {});
+			}
+		}
 		
 		virtual FOnConcertClientSessionTick& OnTick() override { return OnTickDelegate; }
 
@@ -396,6 +422,7 @@ namespace UE::ConcertSyncTests
 		FOnConcertClientSessionTick OnTickDelegate;
 		FConcertServerSessionMock& ServerMock;
 		FGuid EndpointId;
+		bool bIsConnected = true;
 		TMap<FName, TArray<TSharedPtr<IConcertSessionCustomEventHandler>>> CustomEventHandlers;
 	};
 
