Clients can now query the registered replication streams and authority of other clients

diff --git a/Engine/Plugins/Developer/Concert/ConcertSync/ConcertSyncClient/Source/ConcertSyncClient/Private/Replication/Manager/ReplicationManager.cpp b/Engine/Plugins/Developer/Concert/ConcertSync/ConcertSyncClient/Source/ConcertSyncClient/Private/Replication/Manager/ReplicationManager.cpp
--- a/Engine/Plugins/Developer/Concert/ConcertSync/ConcertSyncClient/Source/ConcertSyncClient/Private/Replication/Manager/ReplicationManager.cpp
+++ b/Engine/Plugins/Developer/Concert/ConcertSync/ConcertSyncClient/Source/ConcertSyncClient/Private/Replication/Manager/ReplicationManager.cpp
@@ -70,6 +70,13 @@ namespace UE::ConcertSyncClient::Replication
 		return RejectAll(MoveTemp(Args));
 	}
 
+	TFuture<FClientQueryResponse> FReplicationManager::QueryClientInfo(FClientQueryRequest Args)
+	{
+		return ensureMsgf(CurrentState, TEXT("StartAcceptingJoinRequests should have been called at this point."))
+			? CurrentState->QueryClientInfo(MoveTemp(Args))
+			: MakeFulfilledPromise<FClientQueryResponse>().GetFuture();
+	}
+
 	void FReplicationManager::OnChangeState(TSharedRef<FReplicationManagerState> NewState)
 	{
 		CurrentState = MoveTemp(NewState);
diff --git a/Engine/Plugins/Developer/Concert/ConcertSync/ConcertSyncClient/Source/ConcertSyncClient/Private/Replication/Manager/ReplicationManager.h b/Engine/Plugins/Developer/Concert/ConcertSync/ConcertSyncClient/Source/ConcertSyncClient/Private/Replication/Manager/ReplicationManager.h
--- a/Engine/Plugins/Developer/Concert/ConcertSync/ConcertSyncClient/Source/ConcertSyncClient/Private/Replication/Manager/ReplicationManager.h
+++ b/Engine/Plugins/Developer/Concert/ConcertSync/ConcertSyncClient/Source/ConcertSyncClient/Private/Replication/Manager/ReplicationManager.h
@@ -31,6 +31,7 @@ namespace UE::ConcertSyncClient::Replication
 		virtual bool IsConnectedToReplicationSession() override;
 		virtual EStreamEnumerationResult ForEachRegisteredStream(TFunctionRef<EBreakBehavior(const FReplicationStreamDescription& Stream)> Callback) const override;
 		virtual TFuture<FAuthorityChangeResponse> RequestAuthorityChange(FAuthorityChangeRequest Args) override;
+		virtual TFuture<FClientQueryResponse> QueryClientInfo(FClientQueryRequest Args) override;
 		//~ End IConcertClientReplicationManager Interface
 
 	private:
diff --git a/Engine/Plugins/Developer/Concert/ConcertSync/ConcertSyncClient/Source/ConcertSyncClient/Private/Replication/Manager/ReplicationManagerState_Connected.cpp b/Engine/Plugins/Developer/Concert/ConcertSync/ConcertSyncClient/Source/ConcertSyncClient/Private/Replication/Manager/ReplicationManagerState_Connected.cpp
--- a/Engine/Plugins/Developer/Concert/ConcertSync/ConcertSyncClient/Source/ConcertSyncClient/Private/Replication/Manager/ReplicationManagerState_Connected.cpp
+++ b/Engine/Plugins/Developer/Concert/ConcertSync/ConcertSyncClient/Source/ConcertSyncClient/Private/Replication/Manager/ReplicationManagerState_Connected.cpp
@@ -79,6 +79,21 @@ namespace UE::ConcertSyncClient::Replication
 			});
 	}
 
+	TFuture<FClientQueryResponse> FReplicationManagerState_Connected::QueryClientInfo(FClientQueryRequest Args)
+	{
+		if (EnumHasAllFlags(Args.QueryFlags, EConcertQueryClientStreamFlags::SkipAuthority | EConcertQueryClientStreamFlags::SkipStreamInfo))
+		{
+			UE_LOG(LogConcert, Warning, TEXT("Request QueryClientInfo is pointless because SkipAuthority and SkipStreamInfo are both set. Returning immediately..."));
+			return MakeFulfilledPromise<FClientQueryResponse>().GetFuture();
+		}
+		
+		return LiveSession->SendCustomRequest<FConcertQueryReplicationInfo_Request, FConcertQueryReplicationInfo_Response>(Args, LiveSession->GetSessionServerEndpointId())
+			.Next([](FConcertQueryReplicationInfo_Response&& Response)
+			{
+				return FClientQueryResponse { MoveTemp(Response) };
+			});
+	}
+
 	void FReplicationManagerState_Connected::OnEnterState()
 	{
 		LiveSession->OnTick().AddSP(this, &FReplicationManagerState_Connected::Tick);
diff --git a/Engine/Plugins/Developer/Concert/ConcertSync/ConcertSyncClient/Source/ConcertSyncClient/Private/Replication/Manager/ReplicationManagerState_Connected.h b/Engine/Plugins/Developer/Concert/ConcertSync/ConcertSyncClient/Source/ConcertSyncClient/Private/Replication/Manager/ReplicationManagerState_Connected.h
--- a/Engine/Plugins/Developer/Concert/ConcertSync/ConcertSyncClient/Source/ConcertSyncClient/Private/Replication/Manager/ReplicationManagerState_Connected.h
+++ b/Engine/Plugins/Developer/Concert/ConcertSync/ConcertSyncClient/Source/ConcertSyncClient/Private/Replication/Manager/ReplicationManagerState_Connected.h
@@ -49,6 +49,7 @@ namespace UE::ConcertSyncClient::Replication
 		virtual bool IsConnectedToReplicationSession() override { return true; }
 		virtual EStreamEnumerationResult ForEachRegisteredStream(TFunctionRef<EBreakBehavior(const FReplicationStreamDescription& Stream)> Callback) const override;
 		virtual TFuture<FAuthorityChangeResponse> RequestAuthorityChange(FAuthorityChangeRequest Args) override;
+		virtual TFuture<FClientQueryResponse> QueryClientInfo(FClientQueryRequest Args) override;
 		//~ End IConcertClientReplicationManager Interface
 
 	private:
diff --git a/Engine/Plugins/Developer/Concert/ConcertSync/ConcertSyncClient/Source/ConcertSyncClient/Private/Replication/Manager/ReplicationManagerState_Disconnected.h b/Engine/Plugins/Developer/Concert/ConcertSync/ConcertSyncClient/Source/ConcertSyncClient/Private/Replication/Manager/ReplicationManagerState_Disconnected.h
--- a/Engine/Plugins/Developer/Concert/ConcertSync/ConcertSyncClient/Source/ConcertSyncClient/Private/Replication/Manager/ReplicationManagerState_Disconnected.h
+++ b/Engine/Plugins/Developer/Concert/ConcertSync/ConcertSyncClient/Source/ConcertSyncClient/Private/Replication/Manager/ReplicationManagerState_Disconnected.h
@@ -28,6 +28,7 @@ namespace UE::ConcertSyncClient::Replication
 		virtual bool IsConnectedToReplicationSession() override { return false; }
 		virtual EStreamEnumerationResult ForEachRegisteredStream(TFunctionRef<EBreakBehavior(const FReplicationStreamDescription& Stream)> Callback) const override { return EStreamEnumerationResult::NoRegisteredStreams; }
 		virtual TFuture<FAuthorityChangeResponse> RequestAuthorityChange(FAuthorityChangeRequest Args) override { return RejectAll(MoveTemp(Args)); }
+		virtual TFuture<FClientQueryResponse> QueryClientInfo(FClientQueryRequest Args) override { return MakeFulfilledPromise<FClientQueryResponse>().GetFuture(); }
 		//~ End IConcertClientReplicationManager Interface
 
 	private:
diff --git a/Engine/Plugins/Developer/Concert/ConcertSync/ConcertSyncClient/Source/ConcertSyncClient/Private/Replication/Manager/ReplicationManagerState_Handshaking.h b/Engine/Plugins/Developer/Concert/ConcertSync/ConcertSyncClient/Source/ConcertSyncClient/Private/Replication/Manager/ReplicationManagerState_Handshaking.h
--- a/Engine/Plugins/Developer/Concert/ConcertSync/ConcertSyncClient/Source/ConcertSyncClient/Private/Replication/Manager/ReplicationManagerState_Handshaking.h
+++ b/Engine/Plugins/Developer/Concert/ConcertSync/ConcertSyncClient/Source/ConcertSyncClient/Private/Replication/Manager/ReplicationManagerState_Handshaking.h
@@ -33,6 +33,7 @@ namespace UE::ConcertSyncClient::Replication
 		virtual bool IsConnectedToReplicationSession() override { return false; }
 		virtual EStreamEnumerationResult ForEachRegisteredStream(TFunctionRef<EBreakBehavior(const FReplicationStreamDescription& Stream)> Callback) const override { return EStreamEnumerationResult::NoRegisteredStreams; }
 		virtual TFuture<FAuthorityChangeResponse> RequestAuthorityChange(FAuthorityChangeRequest Args) override { return RejectAll(MoveTemp(Args)); }
+		virtual TFuture<FClientQueryResponse> QueryClientInfo(FClientQueryRequest Args) override { return MakeFulfilledPromise<FClientQueryResponse>().GetFuture(); }
 		//~ End IConcertClientReplicationManager Interface
 
 	private:
diff --git a/Engine/Plugins/Developer/Concert/ConcertSync/ConcertSyncClient/Source/ConcertSyncClient/Public/Replication/IConcertClientReplicationManager.h b/Engine/Plugins/Developer/Concert/ConcertSync/ConcertSyncClient/Source/ConcertSyncClient/Public/Replication/IConcertClientReplicationManager.h
--- a/Engine/Plugins/Developer/Concert/ConcertSync/ConcertSyncClient/Source/ConcertSyncClient/Public/Replication/IConcertClientReplicationManager.h
+++ b/Engine/Plugins/Developer/Concert/ConcertSync/ConcertSyncClient/Source/ConcertSyncClient/Public/Replication/IConcertClientReplicationManager.h
@@ -36,11 +36,11 @@ namespace UE::ConcertSyncClient::Replication
 		{}
 	};
 
-	// The intention here is to wrap the request in case there are some more specific meta data we want to add in the future.
-	struct FAuthorityChangeRequest : FConcertChangeAuthority_Request
-	{};
-	struct FAuthorityChangeResponse : FConcertChangeAuthority_Response
-	{};
+	// The intention here is to wrap the request in case there is some more specific meta data we want to add in the future.
+	struct FAuthorityChangeRequest : FConcertChangeAuthority_Request {};
+	struct FAuthorityChangeResponse : FConcertChangeAuthority_Response {};
+	struct FClientQueryRequest : FConcertQueryReplicationInfo_Request {};
+	struct FClientQueryResponse : FConcertQueryReplicationInfo_Response {};
 }
 
 /**
@@ -57,7 +57,7 @@ public:
 	 * Joins a replication session.
 	 * Subsequent calls to JoinReplicationSession will fail until either the resulting TFuture returns or LeaveReplicationSession is called.
 	 *
-	 * @note This may execute on any thread. Take care to synchronize correctly with the game thread if needed.
+	 * @note The future may execute on any thread. Take care to synchronize correctly with the game thread if needed.
 	 */
 	virtual TFuture<UE::ConcertSyncClient::Replication::FJoinReplicatedSessionResult> JoinReplicationSession(UE::ConcertSyncClient::Replication::FJoinReplicatedSessionArgs Args) = 0;
 	/** Leaves the current replication session. */
@@ -83,13 +83,19 @@ public:
 	
 	/**
 	 * Requests from the server to change the authority over some objects.
-	 * @note This may execute on any thread. Take care of synchronize correctly with the game thread if needed.
+	 * @note The future may execute on any thread. Take care of synchronize correctly with the game thread if needed.
 	 */
 	virtual TFuture<UE::ConcertSyncClient::Replication::FAuthorityChangeResponse> RequestAuthorityChange(UE::ConcertSyncClient::Replication::FAuthorityChangeRequest Args) = 0;
 	/** Util function that will request authority for all streams for the given objects. */
 	TFuture<UE::ConcertSyncClient::Replication::FAuthorityChangeResponse> TakeAuthorityOver(TArrayView<const FSoftObjectPath> Objects);
 	/** Util function that will let go over all authority of the given objects. */
 	TFuture<UE::ConcertSyncClient::Replication::FAuthorityChangeResponse> ReleaseAuthorityOf(TArrayView<const FSoftObjectPath> Objects);
+
+	/**
+	 * Requests replication info about other clients, including the streams registered and which objects they have authority over (i.e. are sending).
+	 * @note The future may execute on any thread. Take care to synchronize correctly with the game thread if needed.
+	 */
+	virtual TFuture<UE::ConcertSyncClient::Replication::FClientQueryResponse> QueryClientInfo(UE::ConcertSyncClient::Replication::FClientQueryRequest Args) = 0; 
 	
 	virtual ~IConcertClientReplicationManager() = default;
 };
diff --git a/Engine/Plugins/Developer/Concert/ConcertSync/ConcertSyncCore/Source/ConcertSyncCore/Public/Replication/Data/ClientQueriedInfo.h b/Engine/Plugins/Developer/Concert/ConcertSync/ConcertSyncCore/Source/ConcertSyncCore/Public/Replication/Data/ClientQueriedInfo.h
new file mode 100644
--- /dev/null
+++ b/Engine/Plugins/Developer/Concert/ConcertSync/ConcertSyncCore/Source/ConcertSyncCore/Public/Replication/Data/ClientQueriedInfo.h
@@ -0,0 +1,48 @@
+// Copyright Epic Games, Inc. All Rights Reserved.
+
+#pragma once
+
+#include "Replication/Data/ReplicationStreamDescription.h"
+#include "ClientQueriedInfo.generated.h"
+
+/** Describes objects a client has authority over */
+USTRUCT()
+struct FReplicationAuthorityInfo
+{
+	GENERATED_BODY()
+
+	/** The stream ID the client has registered */
+	UPROPERTY()
+	FGuid StreamId;
+
+	/**
+	 * The objects the client has authority over.
+	 * You can figure out the properties by inspecting the corresponding client's stream using the object path.
+	 */
+	UPROPERTY()
+	TArray<FSoftObjectPath> AuthoredObjects;
+};
+
+/** This is info a client receives about another client via FConcertQueryClientStreams_Request. */
+USTRUCT()
+struct FReplicationClientQueriedInfo
+{
+	GENERATED_BODY()
+
+	/**
+	 * The streams the client has registered.
+	 * 
+	 * If EConcertQueryClientStreamFlags::SkipStreamInfo was set, this is empty.
+	 * If EConcertQueryClientStreamFlags::SkipProperties was set, streams' FReplicatedObjectInfo::PropertySelection are empty.
+	 */
+	UPROPERTY()
+	TArray<FSharedReplicationStreamDescription> Streams;
+
+	/**
+	 * Indirectly describes which object properties the client has authority over.
+	 *
+	 * If EConcertQueryClientStreamFlags::SkipAuthority was set, this is empty.
+	 */
+	UPROPERTY()
+	TArray<FReplicationAuthorityInfo> Authority;
+};
\ No newline at end of file
diff --git a/Engine/Plugins/Developer/Concert/ConcertSync/ConcertSyncCore/Source/ConcertSyncCore/Public/Replication/Data/ConcertPropertySelection.h b/Engine/Plugins/Developer/Concert/ConcertSync/ConcertSyncCore/Source/ConcertSyncCore/Public/Replication/Data/ConcertPropertySelection.h
--- a/Engine/Plugins/Developer/Concert/ConcertSync/ConcertSyncCore/Source/ConcertSyncCore/Public/Replication/Data/ConcertPropertySelection.h
+++ b/Engine/Plugins/Developer/Concert/ConcertSync/ConcertSyncCore/Source/ConcertSyncCore/Public/Replication/Data/ConcertPropertySelection.h
@@ -183,6 +183,15 @@ struct CONCERTSYNCCORE_API FConcertPropertySelection
 
 	/** @return Whether this and Other contain at least one property that is the same. This algorithm is strictly O(n^2) but runs O(n) on average. */
 	bool OverlapsWith(const FConcertPropertySelection& Other) const;
+
+	friend bool operator==(const FConcertPropertySelection& Left, const FConcertPropertySelection& Right)
+	{
+		return Left.ReplicatedProperties == Right.ReplicatedProperties;
+	}
+	friend bool operator!=(const FConcertPropertySelection& Left, const FConcertPropertySelection& Right)
+	{
+		return !(Left == Right);
+	}
 };
 
 CONCERTSYNCCORE_API uint32 GetTypeHash(const FConcertPropertyChain& Chain);
\ No newline at end of file
diff --git a/Engine/Plugins/Developer/Concert/ConcertSync/ConcertSyncCore/Source/ConcertSyncCore/Public/Replication/Data/ObjectReplicationMap.h b/Engine/Plugins/Developer/Concert/ConcertSync/ConcertSyncCore/Source/ConcertSyncCore/Public/Replication/Data/ObjectReplicationMap.h
--- a/Engine/Plugins/Developer/Concert/ConcertSync/ConcertSyncCore/Source/ConcertSyncCore/Public/Replication/Data/ObjectReplicationMap.h
+++ b/Engine/Plugins/Developer/Concert/ConcertSync/ConcertSyncCore/Source/ConcertSyncCore/Public/Replication/Data/ObjectReplicationMap.h
@@ -15,6 +15,17 @@ struct FReplicatedObjectInfo
 	
 	UPROPERTY()
 	FConcertPropertySelection PropertySelection;
+
+	friend bool operator==(const FReplicatedObjectInfo& Left, const FReplicatedObjectInfo& Right)
+	{
+		return Left.ClassPath == Right.ClassPath
+			&& Left.PropertySelection == Right.PropertySelection;
+	}
+
+	friend bool operator!=(const FReplicatedObjectInfo& Left, const FReplicatedObjectInfo& Right)
+	{
+		return !(Left == Right);
+	}
 };
 
 /** Maps objects to their replicated properties. */
@@ -31,4 +42,13 @@ struct FObjectReplicationMap
 	 */
 	UPROPERTY()
 	TMap<FSoftObjectPath, FReplicatedObjectInfo> ReplicatedObjects;
+
+	friend bool operator==(const FObjectReplicationMap& Left, const FObjectReplicationMap& Right)
+	{
+		return Left.ReplicatedObjects.OrderIndependentCompareEqual(Right.ReplicatedObjects);
+	}
+	friend bool operator!=(const FObjectReplicationMap& Left, const FObjectReplicationMap& Right)
+	{
+		return !(Left == Right);
+	}
 };
\ No newline at end of file
diff --git a/Engine/Plugins/Developer/Concert/ConcertSync/ConcertSyncCore/Source/ConcertSyncCore/Public/Replication/Messages/ConcertReplicationEvents.h b/Engine/Plugins/Developer/Concert/ConcertSync/ConcertSyncCore/Source/ConcertSyncCore/Public/Replication/Messages/ConcertReplicationEvents.h
--- a/Engine/Plugins/Developer/Concert/ConcertSync/ConcertSyncCore/Source/ConcertSyncCore/Public/Replication/Messages/ConcertReplicationEvents.h
+++ b/Engine/Plugins/Developer/Concert/ConcertSync/ConcertSyncCore/Source/ConcertSyncCore/Public/Replication/Messages/ConcertReplicationEvents.h
@@ -3,6 +3,7 @@
 #pragma once
 
 #include "ConcertMessageData.h"
+#include "Replication/Data/ClientQueriedInfo.h"
 #include "ConcertReplicationEvents.generated.h"
 
 /** Contains data to be applied to a replicated object. */
@@ -113,4 +114,68 @@ struct FConcertChangeAuthority_Response
 	 */
 	UPROPERTY()
 	TMap<FSoftObjectPath, FConcertStreamArray> RejectedObjects;
+};
+
+/** Flags to affect FConcertQueryClientStreams_Request */
+UENUM()
+enum class EConcertQueryClientStreamFlags : uint8
+{
+	None = 0,
+
+	/**
+	 * If set, FReplicationClientQueriedInfo::Streams will be empty.
+	 */
+	SkipStreamInfo = 1 << 0,
+
+	/**
+	 * Optimization. The returned streams will not contain the properties that are sent.
+	 * Set this if the request only cares for the objects being replicated but not its content.
+	 * Only has an effect if SkipStreamInfo is not set.
+	 */
+	SkipProperties = 1 << 1,
+	
+	/**
+	 * If set, FReplicationClientQueriedInfo::Authority will be empty.
+	 */
+	SkipAuthority = 1 << 2
+};
+ENUM_CLASS_FLAGS(EConcertQueryClientStreamFlags)
+
+/**
+ * Let's a client query about another client's registered streams.
+ *
+ * Currently, clients must continuously poll to detect when a client changes this data.
+ * TODO: Add an event with which clients can be notified when another client changes their profile.
+ * TODO: Alternative to an event, we can add a version int which is incremented every time the info changes.
+ */
+USTRUCT()
+struct FConcertQueryReplicationInfo_Request
+{
+	GENERATED_BODY()
+
+	/** The clients to query about */
+	UPROPERTY()
+	TSet<FGuid> ClientEndpointIds;
+
+	UPROPERTY()
+	EConcertQueryClientStreamFlags QueryFlags = EConcertQueryClientStreamFlags::None;
+};
+
+USTRUCT()
+struct FConcertQueryReplicationInfo_Response
+{
+	GENERATED_BODY()
+
+	/**
+	 * Binds client participating in the replication session to their info.
+	 * 
+	 * @note This contains only the client endpoint IDs that are in the replication session at the time of processing the request.
+	 * That means not every endpoint ID that was in the request's ClientEndpointIds will necessarily be in this result. For example,
+	 * if client A sends a FConcertQueryReplicationInfo_Request while client B disconnects, client B may not show up in ClientInfo.
+	 * 
+	 * Key: Client endpoint ID.
+	 * Value: Info about the client.
+	 */
+	UPROPERTY()
+	TMap<FGuid, FReplicationClientQueriedInfo> ClientInfo;
 };
\ No newline at end of file
diff --git a/Engine/Plugins/Developer/Concert/ConcertSync/ConcertSyncServer/Source/ConcertSyncServer/Private/Replication/ConcertServerReplicationManager.cpp b/Engine/Plugins/Developer/Concert/ConcertSync/ConcertSyncServer/Source/ConcertSyncServer/Private/Replication/ConcertServerReplicationManager.cpp
--- a/Engine/Plugins/Developer/Concert/ConcertSync/ConcertSyncServer/Source/ConcertSyncServer/Private/Replication/ConcertServerReplicationManager.cpp
+++ b/Engine/Plugins/Developer/Concert/ConcertSync/ConcertSyncServer/Source/ConcertSyncServer/Private/Replication/ConcertServerReplicationManager.cpp
@@ -21,6 +21,7 @@ namespace UE::ConcertSyncServer::Replication
 		, ReplicationDataReceiver(AuthorityManager, Session, ReplicationCache)
 	{
 		Session->RegisterCustomRequestHandler<FConcertReplication_Join_Request, FConcertReplication_Join_Response>(this, &FConcertServerReplicationManager::HandleJoinReplicationSessionRequest);
+		Session->RegisterCustomRequestHandler<FConcertQueryReplicationInfo_Request, FConcertQueryReplicationInfo_Response>(this, &FConcertServerReplicationManager::HandleQueryReplicationInfoRequest);
 		Session->RegisterCustomEventHandler<FConcertReplication_LeaveEvent>(this, &FConcertServerReplicationManager::HandleLeaveReplicationSessionRequest);
 		Session->OnSessionClientChanged().AddRaw(this, &FConcertServerReplicationManager::OnConnectionChanged);
 
@@ -30,6 +31,7 @@ namespace UE::ConcertSyncServer::Replication
 	FConcertServerReplicationManager::~FConcertServerReplicationManager()
 	{
 		Session->UnregisterCustomRequestHandler<FConcertReplication_Join_Response>();
+		Session->UnregisterCustomRequestHandler<FConcertQueryReplicationInfo_Response>();
 		Session->UnregisterCustomEventHandler<FConcertReplication_LeaveEvent>(this);
 
 		Session->OnTick().RemoveAll(this);
@@ -112,6 +114,94 @@ namespace UE::ConcertSyncServer::Replication
 		return EConcertSessionResponseCode::Success;
 	}
 
+	EConcertSessionResponseCode FConcertServerReplicationManager::HandleQueryReplicationInfoRequest(
+		const FConcertSessionContext& ConcertSessionContext,
+		const FConcertQueryReplicationInfo_Request& Request,
+		FConcertQueryReplicationInfo_Response& Response
+		)
+	{
+		for (const FGuid& EndpointId : Request.ClientEndpointIds)
+		{
+			const TSharedRef<FConcertReplicationClient>* Client = Clients.Find(EndpointId);
+			if (!Client)
+			{
+				// This could happen if the client left the replication session before this request was answered
+				continue;
+			}
+			
+			FReplicationClientQueriedInfo& EndpointInfo = Response.ClientInfo.Add(EndpointId);
+			if (!EnumHasAnyFlags(Request.QueryFlags, EConcertQueryClientStreamFlags::SkipStreamInfo))
+			{
+				const bool bSkipProperties = EnumHasAnyFlags(Request.QueryFlags, EConcertQueryClientStreamFlags::SkipProperties);
+				EndpointInfo.Streams = BuildClientStreamInfo(Client->Get(), bSkipProperties);
+			}
+			if (!EnumHasAnyFlags(Request.QueryFlags, EConcertQueryClientStreamFlags::SkipAuthority))
+			{
+				EndpointInfo.Authority = BuildClientAuthorityInfo(Client->Get());
+			}
+		}
+
+		return EConcertSessionResponseCode::Success;
+	}
+
+	TArray<FSharedReplicationStreamDescription> FConcertServerReplicationManager::BuildClientStreamInfo(const FConcertReplicationClient& Client, bool bSkipProperties) const
+	{
+		TArray<FSharedReplicationStreamDescription> Result;
+		Algo::Transform(Client.GetStreamDescriptions(), Result, [bSkipProperties](const FReplicationStreamDescription& Description)
+		{
+			if (!bSkipProperties)
+			{
+				return Description.BaseDescription;
+			}
+
+			const FSharedReplicationStreamDescription& CopiedDescription = Description.BaseDescription;
+			const FObjectReplicationMap& CopiedReplicationMap = CopiedDescription.ReplicationMap;
+
+			// It was requested to skip sending the properties (saves network bandwidth)
+			FSharedReplicationStreamDescription StrippedResult;
+			StrippedResult.Identifier = CopiedDescription.Identifier;
+			StrippedResult.ReplicationMap.ReplicatedObjects.Reserve(CopiedReplicationMap.ReplicatedObjects.Num());
+			for (const TPair<FSoftObjectPath, FReplicatedObjectInfo>& ObjectInfo : CopiedReplicationMap.ReplicatedObjects)
+			{
+				StrippedResult.ReplicationMap.ReplicatedObjects.Add(ObjectInfo.Key, { ObjectInfo.Value.ClassPath });
+			}
+			return StrippedResult;
+		});
+
+		return Result;
+	}
+
+	TArray<FReplicationAuthorityInfo> FConcertServerReplicationManager::BuildClientAuthorityInfo(const FConcertReplicationClient& Client) const
+	{
+		TArray<FReplicationAuthorityInfo> Result;
+		for (const FReplicationStreamDescription& Description : Client.GetStreamDescriptions())
+		{
+			const FGuid StreamId = Description.BaseDescription.Identifier;
+			FReplicationAuthorityInfo Info;
+			Info.StreamId = StreamId;
+			
+			for (const TPair<FSoftObjectPath, FReplicatedObjectInfo>& Pair : Description.BaseDescription.ReplicationMap.ReplicatedObjects)
+			{
+				ConcertSyncCore::FReplicatedObjectInfo ObjectInfo;
+				ObjectInfo.SenderEndpointId = Client.GetClientEndpointId();
+				ObjectInfo.Object = Pair.Key;
+				ObjectInfo.StreamId = StreamId;
+				
+				if (AuthorityManager->IsObjectChangeAllowed(ObjectInfo))
+				{
+					Info.AuthoredObjects.Add(Pair.Key);
+				}
+			}
+
+			// There is no point in sending empty data
+			if (!Info.AuthoredObjects.IsEmpty())
+			{
+				Result.Add(Info);
+			}
+		}
+		return Result;
+	}
+
 	void FConcertServerReplicationManager::HandleLeaveReplicationSessionRequest(
 		const FConcertSessionContext& ConcertSessionContext,
 		const FConcertReplication_LeaveEvent& EventData
diff --git a/Engine/Plugins/Developer/Concert/ConcertSync/ConcertSyncServer/Source/ConcertSyncServer/Private/Replication/ConcertServerReplicationManager.h b/Engine/Plugins/Developer/Concert/ConcertSync/ConcertSyncServer/Source/ConcertSyncServer/Private/Replication/ConcertServerReplicationManager.h
--- a/Engine/Plugins/Developer/Concert/ConcertSync/ConcertSyncServer/Source/ConcertSyncServer/Private/Replication/ConcertServerReplicationManager.h
+++ b/Engine/Plugins/Developer/Concert/ConcertSync/ConcertSyncServer/Source/ConcertSyncServer/Private/Replication/ConcertServerReplicationManager.h
@@ -38,7 +38,6 @@ namespace UE::ConcertSyncServer::Replication
 		, public FNoncopyable
 	{
 	public:
-		
 		explicit FConcertServerReplicationManager(TSharedRef<IConcertServerSession> InLiveSession);
 		virtual ~FConcertServerReplicationManager() override;
 
@@ -69,6 +68,13 @@ namespace UE::ConcertSyncServer::Replication
 		// Event handlers
 		EConcertSessionResponseCode HandleJoinReplicationSessionRequest(const FConcertSessionContext& ConcertSessionContext, const FConcertReplication_Join_Request& Request, FConcertReplication_Join_Response& Response);
 		EConcertSessionResponseCode InternalHandleJoinReplicationSessionRequest(const FConcertSessionContext& ConcertSessionContext, const FConcertReplication_Join_Request& Request, FConcertReplication_Join_Response& Response);
+
+		EConcertSessionResponseCode HandleQueryReplicationInfoRequest(const FConcertSessionContext& ConcertSessionContext, const FConcertQueryReplicationInfo_Request& Request, FConcertQueryReplicationInfo_Response& Response);
+		/** Gets all registered streams and optionally removes the properties. */
+		TArray<FSharedReplicationStreamDescription> BuildClientStreamInfo(const FConcertReplicationClient& Client, bool bSkipProperties) const;
+		/** Maps the client's streams to the objects in that stream the client has taken authority over. */
+		TArray<FReplicationAuthorityInfo> BuildClientAuthorityInfo(const FConcertReplicationClient& Client) const;
+		
 		void HandleLeaveReplicationSessionRequest(const FConcertSessionContext& ConcertSessionContext, const FConcertReplication_LeaveEvent& EventData);
 		
 		void OnConnectionChanged(IConcertServerSession& ConcertServerSession, EConcertClientStatus ConcertClientStatus, const FConcertSessionClientInfo& ConcertSessionClientInfo);
diff --git a/Engine/Plugins/Developer/Concert/ConcertSync/ConcertSyncTest/Source/ConcertSyncTest/Private/Replication/ReplicationStreamRequestTests.cpp b/Engine/Plugins/Developer/Concert/ConcertSync/ConcertSyncTest/Source/ConcertSyncTest/Private/Replication/ReplicationStreamRequestTests.cpp
new file mode 100644
--- /dev/null
+++ b/Engine/Plugins/Developer/Concert/ConcertSync/ConcertSyncTest/Source/ConcertSyncTest/Private/Replication/ReplicationStreamRequestTests.cpp
@@ -0,0 +1,157 @@
+// Copyright Epic Games, Inc. All Rights Reserved.
+
+#include "Util/ClientServerCommunicationTest.h"
+
+#include "Replication/Data/ReplicationStreamDescription.h"
+#include "Replication/IConcertClientReplicationManager.h"
+#include "TestReflectionObject.h"
+#include "Util/SendReceiveObjectTestBase.h"
+
+#include "Components/StaticMeshComponent.h"
+#include "Misc/AutomationTest.h"
+#include "Misc/Paths.h"
+#include "HAL/FileManager.h"
+
+namespace UE::ConcertSyncTests::Replication::Stream
+{
+	/**
+	 * Let's receiving client query the server what the sending client is sending. Checks
+	 * - the reported streams,
+	 * - that the authority updates,
+	 * - the EConcertQueryClientStreamFlags flags
+	 */
+	IMPLEMENT_CUSTOM_SIMPLE_AUTOMATION_TEST(FQueryOtherClientStreams, FSendReceiveObjectTestBase, "Concert.Replication.QueryClientStream", EAutomationTestFlags::EditorContext | EAutomationTestFlags::EngineFilter);
+	bool FQueryOtherClientStreams::RunTest(const FString& Parameters)
+	{
+		// 1. Init
+		SetUpClientAndServer();
+		
+		const FGuid SenderEndpointId = Client_Sender->ClientSessionMock->GetSessionClientEndpointId();
+		const FSoftObjectPath TestObjectPath{ TestObject };
+		ConcertSyncClient::Replication::FClientQueryRequest Request;
+		Request.ClientEndpointIds = { SenderEndpointId };
+		auto TestReplicationMapContent = [this, &SenderEndpointId](const ConcertSyncClient::Replication::FClientQueryResponse& Response)
+		{
+			const FReplicationClientQueriedInfo* Info = Response.ClientInfo.Find(SenderEndpointId);
+			if (!Info)
+			{
+				AddError(TEXT("No info about sending client!"));
+				return;
+			}
+				
+			TestEqual(TEXT("Only received info about requested client"), Response.ClientInfo.Num(), 1);
+			TestEqual(TEXT("Exactly one stream data"), Info->Streams.Num(), 1);
+			if (Info->Streams.IsEmpty())
+			{
+				return;
+			}
+				
+			const FSharedReplicationStreamDescription& StreamDescription = Info->Streams[0];
+			TestEqual(TEXT("StreamDescription->Identifier correct"), StreamDescription.Identifier, SenderStreamId);
+			TestEqual(TEXT("StreamDescription->ReplicationMap has exactly 1 object"), StreamDescription.ReplicationMap.ReplicatedObjects.Num(), 1);
+
+			const FObjectReplicationMap ExpectedReplicationMap = CreateSenderArgs().Streams[0].BaseDescription.ReplicationMap;
+			TestEqual(TEXT("Registered and reported replication maps match"), StreamDescription.ReplicationMap, ExpectedReplicationMap);
+		};
+		
+		// 2.1 Request before taking authority
+		bool bReceivedFirstQueryResponse = false;
+		ClientReplicationManager_Receiver->QueryClientInfo(Request)
+			.Next([this, &SenderEndpointId, &TestReplicationMapContent, &bReceivedFirstQueryResponse](ConcertSyncClient::Replication::FClientQueryResponse&& Response) mutable
+			{
+				bReceivedFirstQueryResponse = true;
+				TestReplicationMapContent(Response);
+				if (const FReplicationClientQueriedInfo* Info = Response.ClientInfo.Find(SenderEndpointId))
+				{
+					TestEqual(TEXT("Contains no authority data"), Info->Authority.Num(), 0);
+				}
+			});
+		TestTrue(TEXT("Received query response (before taking authority)"), bReceivedFirstQueryResponse);
+
+		// 2.2 Request after taking authority
+		bool bSenderReceivedResponse = false;
+		ClientReplicationManager_Sender->TakeAuthorityOver({ TestObjectPath })
+			.Next([this, &bSenderReceivedResponse](const ConcertSyncClient::Replication::FAuthorityChangeResponse& Response) mutable
+			{
+				bSenderReceivedResponse = true;
+				TestEqual(TEXT("No rejection taking authority"), Response.RejectedObjects.Num(), 0); 
+			});
+
+		bool bReceivedSecondQueryResponse = false;
+		ClientReplicationManager_Receiver->QueryClientInfo(Request)
+			.Next([this, &SenderEndpointId, &TestReplicationMapContent, &TestObjectPath, &bReceivedSecondQueryResponse](ConcertSyncClient::Replication::FClientQueryResponse&& Response) mutable
+			{
+				bReceivedSecondQueryResponse = true;
+				TestReplicationMapContent(Response);
+				
+				const FReplicationClientQueriedInfo* Info = Response.ClientInfo.Find(SenderEndpointId);
+				if (!Info || Info->Authority.IsEmpty())
+				{
+					AddError(TEXT("Missing authority data"));
+					return;
+				}
+
+				TestEqual(TEXT("Exactly 1 authority stream"), Info->Authority.Num(), 1);
+				const FReplicationAuthorityInfo& AuthorityInfo = Info->Authority[0];
+				TestEqual(TEXT("Authority stream ID matches registered stream ID"), AuthorityInfo.StreamId, SenderStreamId);
+				TestEqual(TEXT("Has authority over exactly 1 object"), AuthorityInfo.AuthoredObjects.Num(), 1);
+				TestTrue(TEXT("Has authority over registered object"), AuthorityInfo.AuthoredObjects.Contains(TestObjectPath));
+			});
+		TestTrue(TEXT("Received response to taking authority"), bSenderReceivedResponse);
+		TestTrue(TEXT("Received query response (after taking authority)"), bReceivedSecondQueryResponse);
+
+		// 2.3 SkipStreamInfo
+		bool bReceivedResponse_SkipStreamInfo = false;
+		Request.QueryFlags = EConcertQueryClientStreamFlags::SkipStreamInfo;
+		ClientReplicationManager_Receiver->QueryClientInfo(Request)
+			.Next([this, &SenderEndpointId, &bReceivedResponse_SkipStreamInfo](ConcertSyncClient::Replication::FClientQueryResponse&& Response) mutable
+			{
+				bReceivedResponse_SkipStreamInfo = true;
+				if (const FReplicationClientQueriedInfo* Info = Response.ClientInfo.Find(SenderEndpointId))
+				{
+					TestEqual(TEXT("SkipStreamInfo > No stream data"), Info->Streams.Num(), 0);
+				}
+			});
+		TestTrue(TEXT("bReceivedResponse_SkipStreamInfo"), bReceivedResponse_SkipStreamInfo);
+		
+		// 2.4 SkipProperties
+		bool bReceivedResponse_SkipProperties = false;
+		Request.QueryFlags = EConcertQueryClientStreamFlags::SkipProperties;
+		ClientReplicationManager_Receiver->QueryClientInfo(Request)
+			.Next([this, &SenderEndpointId, &TestObjectPath, &bReceivedResponse_SkipProperties](ConcertSyncClient::Replication::FClientQueryResponse&& Response) mutable
+			{
+				bReceivedResponse_SkipProperties = true;
+				const FReplicationClientQueriedInfo* Info = Response.ClientInfo.Find(SenderEndpointId);
+				if (!Info || Info->Streams.IsEmpty())
+				{
+					AddError(TEXT("SkipProperties > No stream info"));
+					return;
+				}
+
+				const FReplicatedObjectInfo* ObjectInfo = Info->Streams[0].ReplicationMap.ReplicatedObjects.Find(TestObjectPath);
+				if (!ObjectInfo)
+				{
+					AddError(TEXT("SkipProperties > No object info"));
+					return;
+				}
+				TestEqual(TEXT("SkipProperties > No property data"), ObjectInfo->PropertySelection.ReplicatedProperties.Num(), 0);
+			});
+		TestTrue(TEXT("bReceivedResponse_SkipProperties"), bReceivedResponse_SkipProperties);
+		
+		// 2.5 SkipAuthority
+		bool bReceivedResponse_SkipAuthority = false;
+		Request.QueryFlags = EConcertQueryClientStreamFlags::SkipAuthority;
+        ClientReplicationManager_Receiver->QueryClientInfo(Request)
+        	.Next([this, &SenderEndpointId, &bReceivedResponse_SkipAuthority](ConcertSyncClient::Replication::FClientQueryResponse&& Response) mutable
+        	{
+        		bReceivedResponse_SkipAuthority = true;
+        		if (const FReplicationClientQueriedInfo* Info = Response.ClientInfo.Find(SenderEndpointId))
+        		{
+        			TestEqual(TEXT("SkipAuthority > No authority data"), Info->Authority.Num(), 0);
+        		}
+        	});
+		TestTrue(TEXT("bReceivedResponse_SkipAuthority"), bReceivedResponse_SkipAuthority);
+		
+		return true;
+	}
+}
diff --git a/Engine/Plugins/Developer/Concert/ConcertSync/ConcertSyncTest/Source/ConcertSyncTest/Private/Replication/Util/SendReceiveObjectTestBase.h b/Engine/Plugins/Developer/Concert/ConcertSync/ConcertSyncTest/Source/ConcertSyncTest/Private/Replication/Util/SendReceiveObjectTestBase.h
--- a/Engine/Plugins/Developer/Concert/ConcertSync/ConcertSyncTest/Source/ConcertSyncTest/Private/Replication/Util/SendReceiveObjectTestBase.h
+++ b/Engine/Plugins/Developer/Concert/ConcertSync/ConcertSyncTest/Source/ConcertSyncTest/Private/Replication/Util/SendReceiveObjectTestBase.h
@@ -57,7 +57,7 @@ namespace UE::ConcertSyncTests::Replication
 		void SetDifferentValues(UTestReflectionObject& Object);
 		void TestEqualTestValues(UTestReflectionObject& Object, FAutomationTestBase& Test);
 		
-	private:
+	protected:
 		
 		const float SentFloat = 42.f;
 		const FVector SentVector = { 21.f, 84.f, -1.f };
