Rename some util types and expose them to reflection system:
- Rename FStreamedObjectID to FObjectInStreamID
- Rename ConcertSyncCore::FReplicatedObjectInfo to FReplicatedObjectId to avoid name overlap with global, reflected ::FReplicatedObjectInfo
- Expose FObjectInStreamID to reflection system
- Expose FReplicatedObjectId to reflection system

diff --git a/Engine/Plugins/Developer/Concert/ConcertSync/ConcertSyncClient/Source/ConcertSyncClient/Private/Replication/Processing/ClientReplicationDataCollector.cpp b/Engine/Plugins/Developer/Concert/ConcertSync/ConcertSyncClient/Source/ConcertSyncClient/Private/Replication/Processing/ClientReplicationDataCollector.cpp
--- a/Engine/Plugins/Developer/Concert/ConcertSync/ConcertSyncClient/Source/ConcertSyncClient/Private/Replication/Processing/ClientReplicationDataCollector.cpp
+++ b/Engine/Plugins/Developer/Concert/ConcertSync/ConcertSyncClient/Source/ConcertSyncClient/Private/Replication/Processing/ClientReplicationDataCollector.cpp
@@ -6,7 +6,7 @@
 #include "Replication/Formats/IObjectReplicationFormat.h"
 #include "Replication/IConcertClientReplicationBridge.h"
 #include "Replication/ReplicationPropertyFilter.h"
-#include "Replication/ReplicationStreamObjectID.h"
+#include "Replication/Data/ObjectIds.h"
 
 namespace UE::ConcertSyncClient::Replication
 {
@@ -54,7 +54,7 @@ namespace UE::ConcertSyncClient::Replication
 		}
 	}
 
-	void FClientReplicationDataCollector::ForEachPendingObject(TFunctionRef<void(const ConcertSyncCore::FStreamedObjectID&)> ProcessItemFunc) const
+	void FClientReplicationDataCollector::ForEachPendingObject(TFunctionRef<void(const FObjectInStreamID&)> ProcessItemFunc) const
 	{
 		for (const TPair<FSoftObjectPath, TArray<FObjectInfo>>& Pair : ObjectsToReplicate)
 		{
@@ -69,7 +69,7 @@ namespace UE::ConcertSyncClient::Replication
 	}
 
 	bool FClientReplicationDataCollector::ExtractReplicationDataForObject(
-		const ConcertSyncCore::FStreamedObjectID& ObjectToProcess,
+		const FObjectInStreamID& ObjectToProcess,
 		TFunctionRef<void(const FConcertSessionSerializedPayload& Payload)> ProcessCopyable,
 		TFunctionRef<void(FConcertSessionSerializedPayload&& Payload)> ProcessMoveable
 		)
diff --git a/Engine/Plugins/Developer/Concert/ConcertSync/ConcertSyncClient/Source/ConcertSyncClient/Private/Replication/Processing/ClientReplicationDataCollector.h b/Engine/Plugins/Developer/Concert/ConcertSync/ConcertSyncClient/Source/ConcertSyncClient/Private/Replication/Processing/ClientReplicationDataCollector.h
--- a/Engine/Plugins/Developer/Concert/ConcertSync/ConcertSyncClient/Source/ConcertSyncClient/Private/Replication/Processing/ClientReplicationDataCollector.h
+++ b/Engine/Plugins/Developer/Concert/ConcertSync/ConcertSyncClient/Source/ConcertSyncClient/Private/Replication/Processing/ClientReplicationDataCollector.h
@@ -41,9 +41,9 @@ namespace UE::ConcertSyncClient::Replication
 		virtual ~FClientReplicationDataCollector() override;
 
 		//~ Begin IReplicationDataSource Interface
-		virtual void ForEachPendingObject(TFunctionRef<void(const ConcertSyncCore::FStreamedObjectID&)> ProcessItemFunc) const override;
+		virtual void ForEachPendingObject(TFunctionRef<void(const FObjectInStreamID&)> ProcessItemFunc) const override;
 		virtual int32 NumObjects() const override { return NumTrackedObjects; }
-		virtual bool ExtractReplicationDataForObject(const ConcertSyncCore::FStreamedObjectID& Object, TFunctionRef<void(const FConcertSessionSerializedPayload& Payload)> ProcessCopyable, TFunctionRef<void(FConcertSessionSerializedPayload&& Payload)> ProcessMoveable) override;
+		virtual bool ExtractReplicationDataForObject(const FObjectInStreamID& Object, TFunctionRef<void(const FConcertSessionSerializedPayload& Payload)> ProcessCopyable, TFunctionRef<void(FConcertSessionSerializedPayload&& Payload)> ProcessMoveable) override;
 		//~ End IReplicationDataSource Interface
 
 	private:
diff --git a/Engine/Plugins/Developer/Concert/ConcertSync/ConcertSyncClient/Source/ConcertSyncClient/Private/Replication/Processing/ClientReplicationDataQueuer.cpp b/Engine/Plugins/Developer/Concert/ConcertSync/ConcertSyncClient/Source/ConcertSyncClient/Private/Replication/Processing/ClientReplicationDataQueuer.cpp
--- a/Engine/Plugins/Developer/Concert/ConcertSync/ConcertSyncClient/Source/ConcertSyncClient/Private/Replication/Processing/ClientReplicationDataQueuer.cpp
+++ b/Engine/Plugins/Developer/Concert/ConcertSync/ConcertSyncClient/Source/ConcertSyncClient/Private/Replication/Processing/ClientReplicationDataQueuer.cpp
@@ -20,7 +20,7 @@ namespace UE::ConcertSyncClient::Replication
 		: ReplicationBridge(ReplicationBridge)
 	{}
 
-	bool FClientReplicationDataQueuer::WantsToAcceptObject(const ConcertSyncCore::FReplicatedObjectInfo& Object) const
+	bool FClientReplicationDataQueuer::WantsToAcceptObject(const FReplicatedObjectId& Object) const
 	{
 		// The server tries not to send objects to are unavailable but it could happen the object becomes unavailable while data is underway
 		const bool bIsAvailable = ReplicationBridge->IsObjectAvailable(Object.Object);
diff --git a/Engine/Plugins/Developer/Concert/ConcertSync/ConcertSyncClient/Source/ConcertSyncClient/Private/Replication/Processing/ClientReplicationDataQueuer.h b/Engine/Plugins/Developer/Concert/ConcertSync/ConcertSyncClient/Source/ConcertSyncClient/Private/Replication/Processing/ClientReplicationDataQueuer.h
--- a/Engine/Plugins/Developer/Concert/ConcertSync/ConcertSyncClient/Source/ConcertSyncClient/Private/Replication/Processing/ClientReplicationDataQueuer.h
+++ b/Engine/Plugins/Developer/Concert/ConcertSync/ConcertSyncClient/Source/ConcertSyncClient/Private/Replication/Processing/ClientReplicationDataQueuer.h
@@ -20,7 +20,7 @@ namespace UE::ConcertSyncClient::Replication
 		static TSharedRef<FClientReplicationDataQueuer> Make(IConcertClientReplicationBridge* ReplicationBridge, TSharedRef<ConcertSyncCore::FObjectReplicationCache> InReplicationCache);
 		
 		//~ Begin IReplicationCacheUser Interface
-		virtual bool WantsToAcceptObject(const ConcertSyncCore::FReplicatedObjectInfo& Object) const override;
+		virtual bool WantsToAcceptObject(const FReplicatedObjectId& Object) const override;
 		//~ End IReplicationCacheUser Interface
 
 	private:
diff --git a/Engine/Plugins/Developer/Concert/ConcertSync/ConcertSyncCore/Source/ConcertSyncCore/Private/Replication/Data/ObjectIds.cpp b/Engine/Plugins/Developer/Concert/ConcertSync/ConcertSyncCore/Source/ConcertSyncCore/Private/Replication/Data/ObjectIds.cpp
new file mode 100644
--- /dev/null
+++ b/Engine/Plugins/Developer/Concert/ConcertSync/ConcertSyncCore/Source/ConcertSyncCore/Private/Replication/Data/ObjectIds.cpp
@@ -0,0 +1,13 @@
+// Copyright Epic Games, Inc. All Rights Reserved.
+
+#include "Replication/Data/ObjectIds.h"
+
+uint32 GetTypeHash(const FObjectInStreamID& StreamObject)
+{
+	return HashCombineFast(GetTypeHash(StreamObject.Object), GetTypeHash(StreamObject.StreamId));
+}
+
+uint32 GetTypeHash(const FReplicatedObjectId& StreamObject)
+{
+	return HashCombineFast(GetTypeHash(StreamObject.SenderEndpointId), GetTypeHash(static_cast<FObjectInStreamID>(StreamObject)));
+}
\ No newline at end of file
diff --git a/Engine/Plugins/Developer/Concert/ConcertSync/ConcertSyncCore/Source/ConcertSyncCore/Private/Replication/Processing/ObjectReplicationCache.cpp b/Engine/Plugins/Developer/Concert/ConcertSync/ConcertSyncCore/Source/ConcertSyncCore/Private/Replication/Processing/ObjectReplicationCache.cpp
--- a/Engine/Plugins/Developer/Concert/ConcertSync/ConcertSyncCore/Source/ConcertSyncCore/Private/Replication/Processing/ObjectReplicationCache.cpp
+++ b/Engine/Plugins/Developer/Concert/ConcertSync/ConcertSyncCore/Source/ConcertSyncCore/Private/Replication/Processing/ObjectReplicationCache.cpp
@@ -2,7 +2,7 @@
 
 #include "Replication/Processing/ObjectReplicationCache.h"
 
-#include "Replication/ReplicationStreamObjectID.h"
+#include "Replication/Data/ObjectIds.h"
 #include "Replication/Formats/IObjectReplicationFormat.h"
 #include "Replication/Messages/ConcertReplicationEvents.h"
 
@@ -16,7 +16,7 @@ namespace UE::ConcertSyncCore
 	{
 		int32 NumAccepted = 0;
 		
-		const FReplicatedObjectInfo ObjectId{ { OriginStreamId, ObjectReplicationEvent.ReplicatedObject }, SendingEndpointId };
+		const FReplicatedObjectId ObjectId{ { OriginStreamId, ObjectReplicationEvent.ReplicatedObject }, SendingEndpointId };
 		FObjectCache* ObjectCacheBeforeAddition = Cache.Find(ObjectId);
 		if (ObjectCacheBeforeAddition)
 		{
@@ -66,7 +66,7 @@ namespace UE::ConcertSyncCore
 					}
 
 					// ObjectCache may not be found because cache user was unregistered and then destroyed: UnregisterDataCacheUser removes cache users.
-					const FStreamedObjectID ObjectId{ OriginStreamId, LazilyCopiedEventPtr->ReplicatedObject };
+					const FObjectInStreamID ObjectId{ OriginStreamId, LazilyCopiedEventPtr->ReplicatedObject };
 					FObjectCache* ObjectCache = This->Cache.Find(ObjectId);
 					if (!ObjectCache)
 					{
diff --git a/Engine/Plugins/Developer/Concert/ConcertSync/ConcertSyncCore/Source/ConcertSyncCore/Private/Replication/Processing/ObjectReplicationProcessor.cpp b/Engine/Plugins/Developer/Concert/ConcertSync/ConcertSyncCore/Source/ConcertSyncCore/Private/Replication/Processing/ObjectReplicationProcessor.cpp
--- a/Engine/Plugins/Developer/Concert/ConcertSync/ConcertSyncCore/Source/ConcertSyncCore/Private/Replication/Processing/ObjectReplicationProcessor.cpp
+++ b/Engine/Plugins/Developer/Concert/ConcertSync/ConcertSyncCore/Source/ConcertSyncCore/Private/Replication/Processing/ObjectReplicationProcessor.cpp
@@ -14,7 +14,7 @@ namespace UE::ConcertSyncCore
 	void FObjectReplicationProcessor::ProcessObjects(float TimeBudget)
 	{
 		// TODO: Respect time budget and prioritize objects
-		DataSource->ForEachPendingObject([this](const FStreamedObjectID& ObjectInfo)
+		DataSource->ForEachPendingObject([this](const FObjectInStreamID& ObjectInfo)
 		{
 			ProcessObject({ ObjectInfo });
 		});
diff --git a/Engine/Plugins/Developer/Concert/ConcertSync/ConcertSyncCore/Source/ConcertSyncCore/Private/Replication/Processing/ReplicationDataQueuer.cpp b/Engine/Plugins/Developer/Concert/ConcertSync/ConcertSyncCore/Source/ConcertSyncCore/Private/Replication/Processing/ReplicationDataQueuer.cpp
--- a/Engine/Plugins/Developer/Concert/ConcertSync/ConcertSyncCore/Source/ConcertSyncCore/Private/Replication/Processing/ReplicationDataQueuer.cpp
+++ b/Engine/Plugins/Developer/Concert/ConcertSync/ConcertSyncCore/Source/ConcertSyncCore/Private/Replication/Processing/ReplicationDataQueuer.cpp
@@ -14,7 +14,7 @@ namespace UE::ConcertSyncCore
 		}
 	}
 
-	void FReplicationDataQueuer::ForEachPendingObject(TFunctionRef<void(const FStreamedObjectID&)> ProcessItemFunc) const
+	void FReplicationDataQueuer::ForEachPendingObject(TFunctionRef<void(const FObjectInStreamID&)> ProcessItemFunc) const
 	{
 		for (auto It = PendingEvents.CreateConstIterator(); It; ++It)
 		{
@@ -28,7 +28,7 @@ namespace UE::ConcertSyncCore
 	}
 
 	bool FReplicationDataQueuer::ExtractReplicationDataForObject(
-		const FStreamedObjectID& Object,
+		const FObjectInStreamID& Object,
 		TFunctionRef<void(const FConcertSessionSerializedPayload& Payload)> ProcessCopyable,
 		TFunctionRef<void(FConcertSessionSerializedPayload&& Payload)> ProcessMoveable
 		)
@@ -45,7 +45,7 @@ namespace UE::ConcertSyncCore
 		return true;
 	}
 
-	void FReplicationDataQueuer::OnDataCached(const FReplicatedObjectInfo& Object, TSharedRef<const FConcertObjectReplicationEvent> Data)
+	void FReplicationDataQueuer::OnDataCached(const FReplicatedObjectId& Object, TSharedRef<const FConcertObjectReplicationEvent> Data)
 	{
 		PendingEvents.Add(Object, MoveTemp(Data));
 	}
diff --git a/Engine/Plugins/Developer/Concert/ConcertSync/ConcertSyncCore/Source/ConcertSyncCore/Private/Replication/ReplicationStreamObjectID.cpp b/Engine/Plugins/Developer/Concert/ConcertSync/ConcertSyncCore/Source/ConcertSyncCore/Private/Replication/ReplicationStreamObjectID.cpp
deleted file mode 100644
diff --git a/Engine/Plugins/Developer/Concert/ConcertSync/ConcertSyncCore/Source/ConcertSyncCore/Public/Replication/Data/ObjectIds.h b/Engine/Plugins/Developer/Concert/ConcertSync/ConcertSyncCore/Source/ConcertSyncCore/Public/Replication/Data/ObjectIds.h
new file mode 100644
--- /dev/null
+++ b/Engine/Plugins/Developer/Concert/ConcertSync/ConcertSyncCore/Source/ConcertSyncCore/Public/Replication/Data/ObjectIds.h
@@ -0,0 +1,55 @@
+// Copyright Epic Games, Inc. All Rights Reserved.
+
+#pragma once
+
+#include "Misc/Guid.h"
+#include "UObject/SoftObjectPath.h"
+#include "ObjectIds.generated.h"
+
+/** Utility for identifying an object part of a replication stream definition.*/
+USTRUCT()
+struct FObjectInStreamID
+{
+	GENERATED_BODY()
+	
+	/** The replication stream that produces data for this object. */
+	UPROPERTY()
+	FGuid StreamId;
+	/** The object to process */
+	UPROPERTY()
+	FSoftObjectPath Object;
+		
+	friend bool operator==(const FObjectInStreamID& Left, const FObjectInStreamID& Right)
+	{
+		return Left.StreamId == Right.StreamId && Left.Object == Right.Object;
+	}
+
+	friend bool operator!=(const FObjectInStreamID& Left, const FObjectInStreamID& Right)
+	{
+		return !(Left == Right);
+	}
+};
+
+/** Identifies an object that was replicated by a client based on an underlying stream. */
+USTRUCT()
+struct FReplicatedObjectId : public FObjectInStreamID
+{
+	GENERATED_BODY()
+	
+	/** The ID of the endpoint that sent this object. This can be the client or server endpoint depending on who receives it. */
+	FGuid SenderEndpointId;
+		
+	friend bool operator==(const FReplicatedObjectId& Left, const FReplicatedObjectId& Right)
+	{
+		return Left.SenderEndpointId == Right.SenderEndpointId
+			&& static_cast<const FObjectInStreamID&>(Left) == static_cast<const FObjectInStreamID&>(Right);
+	}
+
+	friend bool operator!=(const FReplicatedObjectId& Left, const FReplicatedObjectId& Right)
+	{
+		return !(Left == Right);
+	}
+};
+
+CONCERTSYNCCORE_API uint32 GetTypeHash(const FObjectInStreamID& StreamObject);
+CONCERTSYNCCORE_API uint32 GetTypeHash(const FReplicatedObjectId& StreamObject);
diff --git a/Engine/Plugins/Developer/Concert/ConcertSync/ConcertSyncCore/Source/ConcertSyncCore/Public/Replication/Processing/IReplicationDataSource.h b/Engine/Plugins/Developer/Concert/ConcertSync/ConcertSyncCore/Source/ConcertSyncCore/Public/Replication/Processing/IReplicationDataSource.h
--- a/Engine/Plugins/Developer/Concert/ConcertSync/ConcertSyncCore/Source/ConcertSyncCore/Public/Replication/Processing/IReplicationDataSource.h
+++ b/Engine/Plugins/Developer/Concert/ConcertSync/ConcertSyncCore/Source/ConcertSyncCore/Public/Replication/Processing/IReplicationDataSource.h
@@ -7,10 +7,10 @@
 #include "Misc/Optional.h"
 #include "Templates/Function.h"
 
+struct FObjectInStreamID;
+
 namespace UE::ConcertSyncCore
 {
-	struct FStreamedObjectID;
-	
 	/**
 	 * Responsible for obtaining replication data.
 	 * 
@@ -23,7 +23,7 @@ namespace UE::ConcertSyncCore
 	public:
 		
 		/** Iterates the objects that must be processed for replication (the result of ExtractReplicationDataForObject MAY return something new). */
-		virtual void ForEachPendingObject(TFunctionRef<void(const FStreamedObjectID&)> ProcessItemFunc) const = 0;
+		virtual void ForEachPendingObject(TFunctionRef<void(const FObjectInStreamID&)> ProcessItemFunc) const = 0;
 		/** @return The number of objects ForEachPendingObject would iterate. Can be used, e.g. for Reserve()-ing a container. */
 		virtual int32 NumObjects() const = 0;
 
@@ -40,12 +40,12 @@ namespace UE::ConcertSyncCore
 		 * @return Whether successful. False indicates the call was invalid to make (ForEachPendingObject did not return Object). True indicates success, even if ProcessCopyable was not called.		
 		 */
 		virtual bool ExtractReplicationDataForObject(
-			const FStreamedObjectID& Object,
+			const FObjectInStreamID& Object,
 			TFunctionRef<void(const FConcertSessionSerializedPayload& Payload)> ProcessCopyable,
 			TFunctionRef<void(FConcertSessionSerializedPayload&& Payload)> ProcessMoveable
 			) = 0;
 		/** Util version for callers that only want to read and do not want to store the payload. */
-		bool ExtractReplicationDataForObject(const FStreamedObjectID& Object, TFunctionRef<void(const FConcertSessionSerializedPayload& Payload)> ProcessCopyable)
+		bool ExtractReplicationDataForObject(const FObjectInStreamID& Object, TFunctionRef<void(const FConcertSessionSerializedPayload& Payload)> ProcessCopyable)
 		{
 			return ExtractReplicationDataForObject(Object, ProcessCopyable, [&ProcessCopyable](FConcertSessionSerializedPayload&& Payload){ ProcessCopyable(Payload); });
 		}
diff --git a/Engine/Plugins/Developer/Concert/ConcertSync/ConcertSyncCore/Source/ConcertSyncCore/Public/Replication/Processing/ObjectReplicationCache.h b/Engine/Plugins/Developer/Concert/ConcertSync/ConcertSyncCore/Source/ConcertSyncCore/Public/Replication/Processing/ObjectReplicationCache.h
--- a/Engine/Plugins/Developer/Concert/ConcertSync/ConcertSyncCore/Source/ConcertSyncCore/Public/Replication/Processing/ObjectReplicationCache.h
+++ b/Engine/Plugins/Developer/Concert/ConcertSync/ConcertSyncCore/Source/ConcertSyncCore/Public/Replication/Processing/ObjectReplicationCache.h
@@ -4,7 +4,7 @@
 
 #include "Containers/Array.h"
 #include "Delegates/DelegateCombinations.h"
-#include "Replication/ReplicationStreamObjectID.h"
+#include "Replication/Data/ObjectIds.h"
 #include "Templates/Function.h"
 #include "Templates/SharedPointer.h"
 
@@ -25,7 +25,7 @@ namespace UE::ConcertSyncCore
 	public:
 
 		/** @return Whether this user is interested in data from this object. */
-		virtual bool WantsToAcceptObject(const FReplicatedObjectInfo& Object) const = 0;
+		virtual bool WantsToAcceptObject(const FReplicatedObjectId& Object) const = 0;
 
 		/**
 		 * Called when data that is interesting to this user becomes available.
@@ -33,7 +33,7 @@ namespace UE::ConcertSyncCore
 		 * The user can keep hold of Data until it is used, at which point it just let's Data get out of scope.
 		 * If new data is received while this user is referencing Data, Data will be combined to contain any new data.
 		 */
-		virtual void OnDataCached(const FReplicatedObjectInfo& Object, TSharedRef<const FConcertObjectReplicationEvent> Data) = 0;
+		virtual void OnDataCached(const FReplicatedObjectId& Object, TSharedRef<const FConcertObjectReplicationEvent> Data) = 0;
 
 		virtual ~IReplicationCacheUser() = default;
 	};
@@ -92,6 +92,6 @@ namespace UE::ConcertSyncCore
 			TMap<TWeakPtr<IReplicationCacheUser>, TWeakPtr<FConcertObjectReplicationEvent>> DataInUse;
 		};
 		/** Maps every object to the events cached for it. */
-		TMap<FStreamedObjectID, FObjectCache> Cache; 
+		TMap<FObjectInStreamID, FObjectCache> Cache; 
 	};
 }
\ No newline at end of file
diff --git a/Engine/Plugins/Developer/Concert/ConcertSync/ConcertSyncCore/Source/ConcertSyncCore/Public/Replication/Processing/ObjectReplicationProcessor.h b/Engine/Plugins/Developer/Concert/ConcertSync/ConcertSyncCore/Source/ConcertSyncCore/Public/Replication/Processing/ObjectReplicationProcessor.h
--- a/Engine/Plugins/Developer/Concert/ConcertSync/ConcertSyncCore/Source/ConcertSyncCore/Public/Replication/Processing/ObjectReplicationProcessor.h
+++ b/Engine/Plugins/Developer/Concert/ConcertSync/ConcertSyncCore/Source/ConcertSyncCore/Public/Replication/Processing/ObjectReplicationProcessor.h
@@ -3,7 +3,7 @@
 #pragma once
 
 #include "IReplicationDataSource.h"
-#include "Replication/ReplicationStreamObjectID.h"
+#include "Replication/Data/ObjectIds.h"
 #include "Templates/SharedPointer.h"
 #include "UObject/SoftObjectPath.h"
 
@@ -37,7 +37,7 @@ namespace UE::ConcertSyncCore
 		struct FObjectProcessArgs
 		{
 			/** Info about the object to process */
-			FStreamedObjectID ObjectInfo;
+			FObjectInStreamID ObjectInfo;
 		};
 
 		FORCEINLINE IReplicationDataSource& GetDataSource() const { return DataSource.Get(); }
diff --git a/Engine/Plugins/Developer/Concert/ConcertSync/ConcertSyncCore/Source/ConcertSyncCore/Public/Replication/Processing/ReplicationDataQueuer.h b/Engine/Plugins/Developer/Concert/ConcertSync/ConcertSyncCore/Source/ConcertSyncCore/Public/Replication/Processing/ReplicationDataQueuer.h
--- a/Engine/Plugins/Developer/Concert/ConcertSync/ConcertSyncCore/Source/ConcertSyncCore/Public/Replication/Processing/ReplicationDataQueuer.h
+++ b/Engine/Plugins/Developer/Concert/ConcertSync/ConcertSyncCore/Source/ConcertSyncCore/Public/Replication/Processing/ReplicationDataQueuer.h
@@ -27,13 +27,13 @@ namespace UE::ConcertSyncCore
 		virtual ~FReplicationDataQueuer() override;
 
 		//~ Begin IReplicationDataSource Interface
-		virtual void ForEachPendingObject(TFunctionRef<void(const FStreamedObjectID&)> ProcessItemFunc) const override;
+		virtual void ForEachPendingObject(TFunctionRef<void(const FObjectInStreamID&)> ProcessItemFunc) const override;
 		virtual int32 NumObjects() const override;
-		virtual bool ExtractReplicationDataForObject(const FStreamedObjectID& Object, TFunctionRef<void(const FConcertSessionSerializedPayload& Payload)> ProcessCopyable, TFunctionRef<void(FConcertSessionSerializedPayload&& Payload)> ProcessMoveable) override;
+		virtual bool ExtractReplicationDataForObject(const FObjectInStreamID& Object, TFunctionRef<void(const FConcertSessionSerializedPayload& Payload)> ProcessCopyable, TFunctionRef<void(FConcertSessionSerializedPayload&& Payload)> ProcessMoveable) override;
 		//~ End IReplicationDataSource Interface
 		
 		//~ Begin IReplicationCacheUser Interface
-		virtual void OnDataCached(const FReplicatedObjectInfo& Object, TSharedRef<const FConcertObjectReplicationEvent> Data) override;
+		virtual void OnDataCached(const FReplicatedObjectId& Object, TSharedRef<const FConcertObjectReplicationEvent> Data) override;
 		//~ End IReplicationCacheUser Interface
 
 	protected:
@@ -44,7 +44,7 @@ namespace UE::ConcertSyncCore
 	private:
 
 		/** Stores events as they are received. */
-		TMap<FStreamedObjectID, TSharedPtr<const FConcertObjectReplicationEvent>> PendingEvents;
+		TMap<FObjectInStreamID, TSharedPtr<const FConcertObjectReplicationEvent>> PendingEvents;
 
 		/** Provides us with replication events and shares them effectively. */
 		TSharedPtr<FObjectReplicationCache> ReplicationCache;
diff --git a/Engine/Plugins/Developer/Concert/ConcertSync/ConcertSyncCore/Source/ConcertSyncCore/Public/Replication/ReplicationStreamObjectID.h b/Engine/Plugins/Developer/Concert/ConcertSync/ConcertSyncCore/Source/ConcertSyncCore/Public/Replication/ReplicationStreamObjectID.h
deleted file mode 100644
diff --git a/Engine/Plugins/Developer/Concert/ConcertSync/ConcertSyncServer/Source/ConcertSyncServer/Private/Replication/AuthorityManager.cpp b/Engine/Plugins/Developer/Concert/ConcertSync/ConcertSyncServer/Source/ConcertSyncServer/Private/Replication/AuthorityManager.cpp
--- a/Engine/Plugins/Developer/Concert/ConcertSync/ConcertSyncServer/Source/ConcertSyncServer/Private/Replication/AuthorityManager.cpp
+++ b/Engine/Plugins/Developer/Concert/ConcertSync/ConcertSyncServer/Source/ConcertSyncServer/Private/Replication/AuthorityManager.cpp
@@ -4,7 +4,7 @@
 
 #include "ConcertLogGlobal.h"
 #include "IConcertSession.h"
-#include "Replication/ReplicationStreamObjectID.h"
+#include "Replication/Data/ObjectIds.h"
 #include "Replication/Data/ReplicationStreamDescription.h"
 #include "Replication/Messages/ConcertReplicationEvents.h"
 
@@ -39,7 +39,7 @@ namespace UE::ConcertSyncServer::Replication
 		Session->UnregisterCustomRequestHandler<FConcertChangeAuthority_Request>();
 	}
 
-	bool FAuthorityManager::IsObjectChangeAllowed(const ConcertSyncCore::FReplicatedObjectInfo& ObjectChange) const
+	bool FAuthorityManager::IsObjectChangeAllowed(const FReplicatedObjectId& ObjectChange) const
 	{
 		const FClientAuthorityData* AuthorityData = ClientAuthorityData.Find(ObjectChange.SenderEndpointId);
 		const TSet<FSoftObjectPath>* OwnedObjects = AuthorityData ? AuthorityData->OwnedObjects.Find(ObjectChange.StreamId) : nullptr;
@@ -110,7 +110,7 @@ namespace UE::ConcertSyncServer::Replication
 		return EConcertSessionResponseCode::Success;
 	}
 
-	bool FAuthorityManager::CanTakeAuthority(const FClientAuthorityData& ClientData, const FClientId& ClientId, const ConcertSyncCore::FStreamedObjectID& Object) const
+	bool FAuthorityManager::CanTakeAuthority(const FClientAuthorityData& ClientData, const FClientId& ClientId, const FObjectInStreamID& Object) const
 	{
 		/*
 		 * This function is unoptimized: it repeatedly iterates through every client's registered streams and properties.
diff --git a/Engine/Plugins/Developer/Concert/ConcertSync/ConcertSyncServer/Source/ConcertSyncServer/Private/Replication/AuthorityManager.h b/Engine/Plugins/Developer/Concert/ConcertSync/ConcertSyncServer/Source/ConcertSyncServer/Private/Replication/AuthorityManager.h
--- a/Engine/Plugins/Developer/Concert/ConcertSync/ConcertSyncServer/Source/ConcertSyncServer/Private/Replication/AuthorityManager.h
+++ b/Engine/Plugins/Developer/Concert/ConcertSync/ConcertSyncServer/Source/ConcertSyncServer/Private/Replication/AuthorityManager.h
@@ -12,14 +12,10 @@ struct FConcertChangeAuthority_Response;
 struct FConcertChangeAuthority_Request;
 struct FConcertPropertySelection;
 struct FConcertSessionContext;
+struct FObjectInStreamID;
+struct FReplicatedObjectId;
 struct FReplicationStreamDescription;
 
-namespace UE::ConcertSyncCore
-{
-	struct FStreamedObjectID;
-	struct FReplicatedObjectInfo;
-}
-
 namespace UE::ConcertSyncServer::Replication
 {
 	/**
@@ -53,7 +49,7 @@ namespace UE::ConcertSyncServer::Replication
 		 * Checks whether the client that sent the identified object had authority to send it.
 		 * @return Whether the server should process the object change.
 		 */
-		bool IsObjectChangeAllowed(const ConcertSyncCore::FReplicatedObjectInfo& ObjectChange) const;
+		bool IsObjectChangeAllowed(const FReplicatedObjectId& ObjectChange) const;
 
 		/** Notifies this manager that the client has left, which means all their authority is now gone. */
 		void OnClientLeft(const FClientId& ClientEndpointId);
@@ -80,7 +76,7 @@ namespace UE::ConcertSyncServer::Replication
 			);
 
 		/** Whether it is legal for this the client identified by ClientId to take control over the object given the stream the client has registered. */
-		bool CanTakeAuthority(const FClientAuthorityData& ClientData, const FClientId& ClientId, const ConcertSyncCore::FStreamedObjectID& Object) const;
+		bool CanTakeAuthority(const FClientAuthorityData& ClientData, const FClientId& ClientId, const FObjectInStreamID& Object) const;
 
 		/** Finds a stream registered with the client by its ID. */
 		const FReplicationStreamDescription* FindClientStreamById(const FClientId& ClientId, const FStreamId& StreamId) const;
diff --git a/Engine/Plugins/Developer/Concert/ConcertSync/ConcertSyncServer/Source/ConcertSyncServer/Private/Replication/ConcertServerReplicationManager.cpp b/Engine/Plugins/Developer/Concert/ConcertSync/ConcertSyncServer/Source/ConcertSyncServer/Private/Replication/ConcertServerReplicationManager.cpp
--- a/Engine/Plugins/Developer/Concert/ConcertSync/ConcertSyncServer/Source/ConcertSyncServer/Private/Replication/ConcertServerReplicationManager.cpp
+++ b/Engine/Plugins/Developer/Concert/ConcertSync/ConcertSyncServer/Source/ConcertSyncServer/Private/Replication/ConcertServerReplicationManager.cpp
@@ -182,7 +182,7 @@ namespace UE::ConcertSyncServer::Replication
 			
 			for (const TPair<FSoftObjectPath, FReplicatedObjectInfo>& Pair : Description.BaseDescription.ReplicationMap.ReplicatedObjects)
 			{
-				ConcertSyncCore::FReplicatedObjectInfo ObjectInfo;
+				FReplicatedObjectId ObjectInfo;
 				ObjectInfo.SenderEndpointId = Client.GetClientEndpointId();
 				ObjectInfo.Object = Pair.Key;
 				ObjectInfo.StreamId = StreamId;
diff --git a/Engine/Plugins/Developer/Concert/ConcertSync/ConcertSyncServer/Source/ConcertSyncServer/Private/Replication/Processing/ServerObjectReplicationReceiver.cpp b/Engine/Plugins/Developer/Concert/ConcertSync/ConcertSyncServer/Source/ConcertSyncServer/Private/Replication/Processing/ServerObjectReplicationReceiver.cpp
--- a/Engine/Plugins/Developer/Concert/ConcertSync/ConcertSyncServer/Source/ConcertSyncServer/Private/Replication/Processing/ServerObjectReplicationReceiver.cpp
+++ b/Engine/Plugins/Developer/Concert/ConcertSync/ConcertSyncServer/Source/ConcertSyncServer/Private/Replication/Processing/ServerObjectReplicationReceiver.cpp
@@ -4,7 +4,7 @@
 
 #include "IConcertSessionHandler.h"
 #include "Replication/AuthorityManager.h"
-#include "Replication/ReplicationStreamObjectID.h"
+#include "Replication/Data/ObjectIds.h"
 #include "Replication/Messages/ConcertReplicationEvents.h"
 
 namespace UE::ConcertSyncServer::Replication
@@ -24,7 +24,7 @@ namespace UE::ConcertSyncServer::Replication
 		const FConcertObjectReplicationEvent& ObjectEvent
 		) const
 	{
-		const ConcertSyncCore::FReplicatedObjectInfo ReplicatedObjectInfo { { StreamEvent.StreamId, ObjectEvent.ReplicatedObject }, SessionContext.SourceEndpointId };
+		const FReplicatedObjectId ReplicatedObjectInfo { { StreamEvent.StreamId, ObjectEvent.ReplicatedObject }, SessionContext.SourceEndpointId };
 		return AuthorityManager->IsObjectChangeAllowed(ReplicatedObjectInfo);
 	}
 }
diff --git a/Engine/Plugins/Developer/Concert/ConcertSync/ConcertSyncServer/Source/ConcertSyncServer/Private/Replication/Processing/ServerReplicationDataQueuer.cpp b/Engine/Plugins/Developer/Concert/ConcertSync/ConcertSyncServer/Source/ConcertSyncServer/Private/Replication/Processing/ServerReplicationDataQueuer.cpp
--- a/Engine/Plugins/Developer/Concert/ConcertSync/ConcertSyncServer/Source/ConcertSyncServer/Private/Replication/Processing/ServerReplicationDataQueuer.cpp
+++ b/Engine/Plugins/Developer/Concert/ConcertSync/ConcertSyncServer/Source/ConcertSyncServer/Private/Replication/Processing/ServerReplicationDataQueuer.cpp
@@ -17,7 +17,7 @@ namespace UE::ConcertSyncServer::Replication
 		return Result;
 	}
 
-	bool FServerReplicationDataQueuer::WantsToAcceptObject(const ConcertSyncCore::FReplicatedObjectInfo& Object) const
+	bool FServerReplicationDataQueuer::WantsToAcceptObject(const FReplicatedObjectId& Object) const
 	{
 		// Do not send back the data to the client that generated it
 		const bool bWasSentByThisClient = OwningClientEndpointId == Object.SenderEndpointId;
diff --git a/Engine/Plugins/Developer/Concert/ConcertSync/ConcertSyncServer/Source/ConcertSyncServer/Private/Replication/Processing/ServerReplicationDataQueuer.h b/Engine/Plugins/Developer/Concert/ConcertSync/ConcertSyncServer/Source/ConcertSyncServer/Private/Replication/Processing/ServerReplicationDataQueuer.h
--- a/Engine/Plugins/Developer/Concert/ConcertSync/ConcertSyncServer/Source/ConcertSyncServer/Private/Replication/Processing/ServerReplicationDataQueuer.h
+++ b/Engine/Plugins/Developer/Concert/ConcertSync/ConcertSyncServer/Source/ConcertSyncServer/Private/Replication/Processing/ServerReplicationDataQueuer.h
@@ -24,7 +24,7 @@ namespace UE::ConcertSyncServer::Replication
 		static TSharedRef<FServerReplicationDataQueuer> Make(const FGuid& OwningClientEndpointId, TSharedRef<ConcertSyncCore::FObjectReplicationCache> InReplicationCache);
 		
 		//~ Begin IReplicationCacheUser Interface
-		virtual bool WantsToAcceptObject(const ConcertSyncCore::FReplicatedObjectInfo& Object) const override;
+		virtual bool WantsToAcceptObject(const FReplicatedObjectId& Object) const override;
 		//~ End IReplicationCacheUser Interface
 
 	private:
diff --git a/Engine/Plugins/Developer/Concert/ConcertSync/ConcertSyncTest/Source/ConcertSyncTest/Private/Replication/ReplicationCacheTests.cpp b/Engine/Plugins/Developer/Concert/ConcertSync/ConcertSyncTest/Source/ConcertSyncTest/Private/Replication/ReplicationCacheTests.cpp
--- a/Engine/Plugins/Developer/Concert/ConcertSync/ConcertSyncTest/Source/ConcertSyncTest/Private/Replication/ReplicationCacheTests.cpp
+++ b/Engine/Plugins/Developer/Concert/ConcertSync/ConcertSyncTest/Source/ConcertSyncTest/Private/Replication/ReplicationCacheTests.cpp
@@ -3,7 +3,7 @@
 #include "Replication/Formats/IObjectReplicationFormat.h"
 #include "Replication/Messages/ConcertReplicationEvents.h"
 #include "Replication/Processing/ObjectReplicationCache.h"
-#include "Replication/ReplicationStreamObjectID.h"
+#include "Replication/Data/ObjectIds.h"
 #include "TestReflectionObject.h"
 #include "Util/ClientServerCommunicationTest.h"
 
@@ -60,22 +60,22 @@ namespace UE::ConcertSyncTests
 		public:
 
 			FAutomationTestBase& Test;
-			ConcertSyncCore::FStreamedObjectID AllowedObject;
+			FObjectInStreamID AllowedObject;
 			EReplicationCacheTestFlags Flags;
 			TSharedPtr<const FConcertObjectReplicationEvent> CachedData;
 
-			FTestReplicationCacheUser(FAutomationTestBase& Test, ConcertSyncCore::FStreamedObjectID AllowedObject, EReplicationCacheTestFlags Flags = EReplicationCacheTestFlags::None)
+			FTestReplicationCacheUser(FAutomationTestBase& Test, FObjectInStreamID AllowedObject, EReplicationCacheTestFlags Flags = EReplicationCacheTestFlags::None)
 				: Test(Test)
 				, AllowedObject(MoveTemp(AllowedObject))
 				, Flags(Flags)
 			{}
 			
-			virtual bool WantsToAcceptObject(const ConcertSyncCore::FReplicatedObjectInfo& Object) const override
+			virtual bool WantsToAcceptObject(const FReplicatedObjectId& Object) const override
 			{
 				return !EnumHasAnyFlags(Flags, EReplicationCacheTestFlags::NeverReceive) && AllowedObject == Object;
 			}
 			
-			virtual void OnDataCached(const ConcertSyncCore::FReplicatedObjectInfo& Object, TSharedRef<const FConcertObjectReplicationEvent> Data) override
+			virtual void OnDataCached(const FReplicatedObjectId& Object, TSharedRef<const FConcertObjectReplicationEvent> Data) override
 			{
 				if (EnumHasAnyFlags(Flags, EReplicationCacheTestFlags::NeverReceive))
 				{
@@ -124,7 +124,7 @@ namespace UE::ConcertSyncTests
 		const FGuid StreamId = FGuid::NewGuid();
 		const FSoftObjectPath ObjectPath(TEXT("/Game/World.World:PersistentLevel.StaticMeshActor0"));
 		const FGuid DummySendingClientId = FGuid::NewGuid();
-		const ConcertSyncCore::FReplicatedObjectInfo ObjectID{ { StreamId, ObjectPath }, DummySendingClientId};
+		const FReplicatedObjectId ObjectID{ { StreamId, ObjectPath }, DummySendingClientId};
 		
 		TSharedRef<ConcertSyncCore::FObjectReplicationCache> Cache = MakeShared<ConcertSyncCore::FObjectReplicationCache>(MakeShared<FTestReplicationFormat>());
 		TSharedRef<FTestReplicationCacheUser> User_NeverConsume = MakeShared<FTestReplicationCacheUser>(*this, ObjectID, EReplicationCacheTestFlags::NeverConsume);
