Add way to change streams dynamically during session

diff --git a/Engine/Plugins/Developer/Concert/ConcertSync/ConcertSyncClient/Source/ConcertSyncClient/Private/Replication/IConcertClientReplicationManager.cpp b/Engine/Plugins/Developer/Concert/ConcertSync/ConcertSyncClient/Source/ConcertSyncClient/Private/Replication/IConcertClientReplicationManager.cpp
--- a/Engine/Plugins/Developer/Concert/ConcertSync/ConcertSyncClient/Source/ConcertSyncClient/Private/Replication/IConcertClientReplicationManager.cpp
+++ b/Engine/Plugins/Developer/Concert/ConcertSync/ConcertSyncClient/Source/ConcertSyncClient/Private/Replication/IConcertClientReplicationManager.cpp
@@ -35,6 +35,13 @@ bool IConcertClientReplicationManager::HasRegisteredStreams() const
 	return ForEachRegisteredStream([](const auto&){ return EBreakBehavior::Break; }) == EStreamEnumerationResult::Iterated;
 }
 
+TArray<FReplicationStreamDescription> IConcertClientReplicationManager::GetRegisteredStreams() const
+{
+	TArray<FReplicationStreamDescription> Result;
+	ForEachRegisteredStream([&Result](const FReplicationStreamDescription& Description){ Result.Add(Description); return EBreakBehavior::Continue; });
+	return Result;
+}
+
 TFuture<UE::ConcertSyncClient::Replication::FAuthorityChangeResponse> IConcertClientReplicationManager::TakeAuthorityOver(TArrayView<const FSoftObjectPath> Objects)
 {
 	using namespace UE::ConcertSyncClient::Replication;
@@ -59,7 +66,7 @@ TFuture<UE::ConcertSyncClient::Replication::FAuthorityChangeResponse> IConcertCl
 	{
 		// Not only does this warn about incorrect API use at runtime, this also helps debug (incorrectly written) unit tests
 		const FString ObjectsAsString = FString::JoinBy(Objects, TEXT(","), [](const FSoftObjectPath& Path){ return Path.ToString(); });
-		UE_LOG(LogConcert, Warning, TEXT("Local client did not registered any stream for the given objects. This take authority request will not be sent. Objects: %s"), *ObjectsAsString);
+		UE_LOG(LogConcert, Warning, TEXT("Local client did not register any stream for the given objects. This take authority request will not be sent. Objects: %s"), *ObjectsAsString);
 		return MakeFulfilledPromise<FAuthorityChangeResponse>(FAuthorityChangeResponse{}).GetFuture();
 	}
 	
diff --git a/Engine/Plugins/Developer/Concert/ConcertSync/ConcertSyncClient/Source/ConcertSyncClient/Private/Replication/Manager/ReplicationManager.cpp b/Engine/Plugins/Developer/Concert/ConcertSync/ConcertSyncClient/Source/ConcertSyncClient/Private/Replication/Manager/ReplicationManager.cpp
--- a/Engine/Plugins/Developer/Concert/ConcertSync/ConcertSyncClient/Source/ConcertSyncClient/Private/Replication/Manager/ReplicationManager.cpp
+++ b/Engine/Plugins/Developer/Concert/ConcertSync/ConcertSyncClient/Source/ConcertSyncClient/Private/Replication/Manager/ReplicationManager.cpp
@@ -77,6 +77,13 @@ namespace UE::ConcertSyncClient::Replication
 			: MakeFulfilledPromise<FClientQueryResponse>().GetFuture();
 	}
 
+	TFuture<FChangeStreamResponse> FReplicationManager::ChangeStream(FChangeStreamRequest Args)
+	{
+		return ensureMsgf(CurrentState, TEXT("StartAcceptingJoinRequests should have been called at this point."))
+			? CurrentState->ChangeStream(MoveTemp(Args))
+			: MakeFulfilledPromise<FChangeStreamResponse>().GetFuture(); 
+	}
+
 	void FReplicationManager::OnChangeState(TSharedRef<FReplicationManagerState> NewState)
 	{
 		CurrentState = MoveTemp(NewState);
diff --git a/Engine/Plugins/Developer/Concert/ConcertSync/ConcertSyncClient/Source/ConcertSyncClient/Private/Replication/Manager/ReplicationManager.h b/Engine/Plugins/Developer/Concert/ConcertSync/ConcertSyncClient/Source/ConcertSyncClient/Private/Replication/Manager/ReplicationManager.h
--- a/Engine/Plugins/Developer/Concert/ConcertSync/ConcertSyncClient/Source/ConcertSyncClient/Private/Replication/Manager/ReplicationManager.h
+++ b/Engine/Plugins/Developer/Concert/ConcertSync/ConcertSyncClient/Source/ConcertSyncClient/Private/Replication/Manager/ReplicationManager.h
@@ -32,6 +32,7 @@ namespace UE::ConcertSyncClient::Replication
 		virtual EStreamEnumerationResult ForEachRegisteredStream(TFunctionRef<EBreakBehavior(const FReplicationStreamDescription& Stream)> Callback) const override;
 		virtual TFuture<FAuthorityChangeResponse> RequestAuthorityChange(FAuthorityChangeRequest Args) override;
 		virtual TFuture<FClientQueryResponse> QueryClientInfo(FClientQueryRequest Args) override;
+		virtual TFuture<FChangeStreamResponse> ChangeStream(FChangeStreamRequest Args) override;
 		//~ End IConcertClientReplicationManager Interface
 
 	private:
diff --git a/Engine/Plugins/Developer/Concert/ConcertSync/ConcertSyncClient/Source/ConcertSyncClient/Private/Replication/Manager/ReplicationManagerState_Connected.cpp b/Engine/Plugins/Developer/Concert/ConcertSync/ConcertSyncClient/Source/ConcertSyncClient/Private/Replication/Manager/ReplicationManagerState_Connected.cpp
--- a/Engine/Plugins/Developer/Concert/ConcertSync/ConcertSyncClient/Source/ConcertSyncClient/Private/Replication/Manager/ReplicationManagerState_Connected.cpp
+++ b/Engine/Plugins/Developer/Concert/ConcertSync/ConcertSyncClient/Source/ConcertSyncClient/Private/Replication/Manager/ReplicationManagerState_Connected.cpp
@@ -6,6 +6,7 @@
 #include "ConcertTransportMessages.h"
 #include "IConcertSession.h"
 #include "ReplicationManagerState_Disconnected.h"
+#include "Replication/ChangeStreamSharedUtils.h"
 #include "Replication/Formats/FullObjectFormat.h"
 #include "Replication/Messages/ConcertReplicationHandshakeMessages.h"
 #include "Replication/Processing/ClientReplicationDataCollector.h"
@@ -94,6 +95,22 @@ namespace UE::ConcertSyncClient::Replication
 			});
 	}
 
+	TFuture<FChangeStreamResponse> FReplicationManagerState_Connected::ChangeStream(FChangeStreamRequest Args)
+	{
+		return LiveSession->SendCustomRequest<FConcertChangeStream_Request, FConcertChangeStream_Response>(Args, LiveSession->GetSessionServerEndpointId())
+			.Next([WeakThis = TWeakPtr<FReplicationManagerState_Connected>(SharedThis(this)), Args](FConcertChangeStream_Response&& Response)
+			{
+				const TSharedPtr<FReplicationManagerState_Connected> ThisPin = WeakThis.Pin();
+				if (ThisPin && Response.IsSuccess())
+				{
+					// Need to update local cache of streams
+					ConcertSyncCore::Replication::ChangeStreamUtils::ApplyValidatedRequest(Args, ThisPin->RegisteredStreams);
+				}
+				
+				return FChangeStreamResponse { MoveTemp(Response) };
+			});
+	}
+
 	void FReplicationManagerState_Connected::OnEnterState()
 	{
 		LiveSession->OnTick().AddSP(this, &FReplicationManagerState_Connected::Tick);
diff --git a/Engine/Plugins/Developer/Concert/ConcertSync/ConcertSyncClient/Source/ConcertSyncClient/Private/Replication/Manager/ReplicationManagerState_Connected.h b/Engine/Plugins/Developer/Concert/ConcertSync/ConcertSyncClient/Source/ConcertSyncClient/Private/Replication/Manager/ReplicationManagerState_Connected.h
--- a/Engine/Plugins/Developer/Concert/ConcertSync/ConcertSyncClient/Source/ConcertSyncClient/Private/Replication/Manager/ReplicationManagerState_Connected.h
+++ b/Engine/Plugins/Developer/Concert/ConcertSync/ConcertSyncClient/Source/ConcertSyncClient/Private/Replication/Manager/ReplicationManagerState_Connected.h
@@ -50,6 +50,7 @@ namespace UE::ConcertSyncClient::Replication
 		virtual EStreamEnumerationResult ForEachRegisteredStream(TFunctionRef<EBreakBehavior(const FReplicationStreamDescription& Stream)> Callback) const override;
 		virtual TFuture<FAuthorityChangeResponse> RequestAuthorityChange(FAuthorityChangeRequest Args) override;
 		virtual TFuture<FClientQueryResponse> QueryClientInfo(FClientQueryRequest Args) override;
+		virtual TFuture<FChangeStreamResponse> ChangeStream(FChangeStreamRequest Args) override;
 		//~ End IConcertClientReplicationManager Interface
 
 	private:
@@ -58,8 +59,8 @@ namespace UE::ConcertSyncClient::Replication
 		const TSharedRef<IConcertClientSession> LiveSession;
 		/** Passed to FReplicationManagerState_Disconnected */
 		IConcertClientReplicationBridge* const ReplicationBridge;
-		/** The streams that were registered during the handshake. */
-		const TArray<FReplicationStreamDescription> RegisteredStreams;
+		/** The streams this client has registered with the server. */
+		TArray<FReplicationStreamDescription> RegisteredStreams;
 		
 		/** The format this client will use for sending & receiving data. */
 		const TSharedRef<ConcertSyncCore::IObjectReplicationFormat> ReplicationFormat;
diff --git a/Engine/Plugins/Developer/Concert/ConcertSync/ConcertSyncClient/Source/ConcertSyncClient/Private/Replication/Manager/ReplicationManagerState_Disconnected.h b/Engine/Plugins/Developer/Concert/ConcertSync/ConcertSyncClient/Source/ConcertSyncClient/Private/Replication/Manager/ReplicationManagerState_Disconnected.h
--- a/Engine/Plugins/Developer/Concert/ConcertSync/ConcertSyncClient/Source/ConcertSyncClient/Private/Replication/Manager/ReplicationManagerState_Disconnected.h
+++ b/Engine/Plugins/Developer/Concert/ConcertSync/ConcertSyncClient/Source/ConcertSyncClient/Private/Replication/Manager/ReplicationManagerState_Disconnected.h
@@ -29,6 +29,7 @@ namespace UE::ConcertSyncClient::Replication
 		virtual EStreamEnumerationResult ForEachRegisteredStream(TFunctionRef<EBreakBehavior(const FReplicationStreamDescription& Stream)> Callback) const override { return EStreamEnumerationResult::NoRegisteredStreams; }
 		virtual TFuture<FAuthorityChangeResponse> RequestAuthorityChange(FAuthorityChangeRequest Args) override { return RejectAll(MoveTemp(Args)); }
 		virtual TFuture<FClientQueryResponse> QueryClientInfo(FClientQueryRequest Args) override { return MakeFulfilledPromise<FClientQueryResponse>().GetFuture(); }
+		virtual TFuture<FChangeStreamResponse> ChangeStream(FChangeStreamRequest Args) override { return MakeFulfilledPromise<FChangeStreamResponse>().GetFuture(); }
 		//~ End IConcertClientReplicationManager Interface
 
 	private:
diff --git a/Engine/Plugins/Developer/Concert/ConcertSync/ConcertSyncClient/Source/ConcertSyncClient/Private/Replication/Manager/ReplicationManagerState_Handshaking.h b/Engine/Plugins/Developer/Concert/ConcertSync/ConcertSyncClient/Source/ConcertSyncClient/Private/Replication/Manager/ReplicationManagerState_Handshaking.h
--- a/Engine/Plugins/Developer/Concert/ConcertSync/ConcertSyncClient/Source/ConcertSyncClient/Private/Replication/Manager/ReplicationManagerState_Handshaking.h
+++ b/Engine/Plugins/Developer/Concert/ConcertSync/ConcertSyncClient/Source/ConcertSyncClient/Private/Replication/Manager/ReplicationManagerState_Handshaking.h
@@ -34,6 +34,7 @@ namespace UE::ConcertSyncClient::Replication
 		virtual EStreamEnumerationResult ForEachRegisteredStream(TFunctionRef<EBreakBehavior(const FReplicationStreamDescription& Stream)> Callback) const override { return EStreamEnumerationResult::NoRegisteredStreams; }
 		virtual TFuture<FAuthorityChangeResponse> RequestAuthorityChange(FAuthorityChangeRequest Args) override { return RejectAll(MoveTemp(Args)); }
 		virtual TFuture<FClientQueryResponse> QueryClientInfo(FClientQueryRequest Args) override { return MakeFulfilledPromise<FClientQueryResponse>().GetFuture(); }
+		virtual TFuture<FChangeStreamResponse> ChangeStream(FChangeStreamRequest Args) override { return MakeFulfilledPromise<FChangeStreamResponse>().GetFuture(); }
 		//~ End IConcertClientReplicationManager Interface
 
 	private:
diff --git a/Engine/Plugins/Developer/Concert/ConcertSync/ConcertSyncClient/Source/ConcertSyncClient/Public/Replication/IConcertClientReplicationManager.h b/Engine/Plugins/Developer/Concert/ConcertSync/ConcertSyncClient/Source/ConcertSyncClient/Public/Replication/IConcertClientReplicationManager.h
--- a/Engine/Plugins/Developer/Concert/ConcertSync/ConcertSyncClient/Source/ConcertSyncClient/Public/Replication/IConcertClientReplicationManager.h
+++ b/Engine/Plugins/Developer/Concert/ConcertSync/ConcertSyncClient/Source/ConcertSyncClient/Public/Replication/IConcertClientReplicationManager.h
@@ -41,6 +41,8 @@ namespace UE::ConcertSyncClient::Replication
 	struct FAuthorityChangeResponse : FConcertChangeAuthority_Response {};
 	struct FClientQueryRequest : FConcertQueryReplicationInfo_Request {};
 	struct FClientQueryResponse : FConcertQueryReplicationInfo_Response {};
+	struct FChangeStreamRequest : FConcertChangeStream_Request {};
+	struct FChangeStreamResponse : FConcertChangeStream_Response {};
 }
 
 /**
@@ -73,13 +75,15 @@ public:
 
 	enum class EStreamEnumerationResult { NoRegisteredStreams, Iterated };
 	/**
-	 * Iterates the streams that were registered in the JoinReplicationSession request.
+	 * Iterates the streams the client has registered with the server.
 	 * It only makes sense to call this function the manager has joined a replication session.
 	 * @return Whether this manager is connected to a session (Iterated) or not (NoRegisteredStreams).
 	 */
 	virtual EStreamEnumerationResult ForEachRegisteredStream(TFunctionRef<EBreakBehavior(const FReplicationStreamDescription& Stream)> Callback) const = 0;
 	/** @return Whether this manager is currently in a replication session (basically whether ForEachRegisteredStream returns EStreamEnumerationResult::Iterated). */
 	bool HasRegisteredStreams() const;
+	/** @return The streams registered with the server. */
+	TArray<FReplicationStreamDescription> GetRegisteredStreams() const;
 	
 	/**
 	 * Requests from the server to change the authority over some objects.
@@ -95,7 +99,13 @@ public:
 	 * Requests replication info about other clients, including the streams registered and which objects they have authority over (i.e. are sending).
 	 * @note The future may execute on any thread. Take care to synchronize correctly with the game thread if needed.
 	 */
-	virtual TFuture<UE::ConcertSyncClient::Replication::FClientQueryResponse> QueryClientInfo(UE::ConcertSyncClient::Replication::FClientQueryRequest Args) = 0; 
+	virtual TFuture<UE::ConcertSyncClient::Replication::FClientQueryResponse> QueryClientInfo(UE::ConcertSyncClient::Replication::FClientQueryRequest Args) = 0;
+
+	/**
+	 * Requests to change the client's registered streams
+	 * @note The future may execute on any thread. Take care to synchronize correctly with the game thread if needed.
+	 */
+	virtual TFuture<UE::ConcertSyncClient::Replication::FChangeStreamResponse> ChangeStream(UE::ConcertSyncClient::Replication::FChangeStreamRequest Args) = 0;
 	
 	virtual ~IConcertClientReplicationManager() = default;
 };
diff --git a/Engine/Plugins/Developer/Concert/ConcertSync/ConcertSyncCore/Source/ConcertSyncCore/Private/Replication/ChangeStreamSharedUtils.cpp b/Engine/Plugins/Developer/Concert/ConcertSync/ConcertSyncCore/Source/ConcertSyncCore/Private/Replication/ChangeStreamSharedUtils.cpp
new file mode 100644
--- /dev/null
+++ b/Engine/Plugins/Developer/Concert/ConcertSync/ConcertSyncCore/Source/ConcertSyncCore/Private/Replication/ChangeStreamSharedUtils.cpp
@@ -0,0 +1,105 @@
+// Copyright Epic Games, Inc. All Rights Reserved.
+
+#include "Replication/ChangeStreamSharedUtils.h"
+
+#include "Misc/EBreakBehavior.h"
+#include "Replication/Data/ObjectIds.h"
+#include "Replication/Messages/ConcertReplicationEvents.h"
+
+namespace UE::ConcertSyncCore::Replication::ChangeStreamUtils
+{
+	void ForEachObjectLosingAuthority(
+		const FConcertChangeStream_Request& Request,
+		const TArray<FReplicationStreamDescription>& ExistingStreams,
+		TFunctionRef<EBreakBehavior(const FObjectInStreamID&)> Callback
+		)
+	{
+		for (const FReplicationStreamDescription& ExistingStream : ExistingStreams)
+		{
+			const FGuid& StreamId = ExistingStream.BaseDescription.Identifier;
+			if (Request.StreamsToRemove.Contains(StreamId))
+			{
+				for (const TPair<FSoftObjectPath, FReplicatedObjectInfo>& RemovePair : ExistingStream.BaseDescription.ReplicationMap.ReplicatedObjects)
+				{
+					const FObjectInStreamID RemovedObjectId {  StreamId, RemovePair.Key };
+					if (Callback(RemovedObjectId) == EBreakBehavior::Break)
+					{
+						return;
+					}
+				}
+			}
+		}
+
+		// This will call Callback multiple times if the same object is in ObjectsToRemove and StreamsToRemove
+		// but that does not really matter (it is a weird case anyhow - why would anyone construct such a request?)
+		for (const FObjectInStreamID& ObjectToRemove : Request.ObjectsToRemove)
+		{
+			if (Callback(ObjectToRemove) == EBreakBehavior::Break)
+			{
+				return;
+			}
+		}
+	}
+	
+	void ApplyValidatedRequest(const FConcertChangeStream_Request& Request, TArray<FReplicationStreamDescription>& StreamsToModify)
+	{
+		for (auto StreamIt = StreamsToModify.CreateIterator(); StreamIt; ++StreamIt)
+		{
+			FSharedReplicationStreamDescription& BaseDescription = StreamIt->BaseDescription;
+			TMap<FSoftObjectPath, FReplicatedObjectInfo>& ReplicationMap = BaseDescription.ReplicationMap.ReplicatedObjects;
+			const FGuid& StreamId = BaseDescription.Identifier;
+			
+			for (const FObjectInStreamID& ObjectToRemove : Request.ObjectsToRemove)
+			{
+				if (ObjectToRemove.StreamId == StreamId)
+				{
+					ReplicationMap.Remove(ObjectToRemove.Object);
+				}
+			}
+			
+			if (Request.StreamsToRemove.Contains(StreamId)
+				// By contract, ObjectsToRemove will remove the stream if the stream ends up being empty
+				|| ReplicationMap.IsEmpty())
+			{
+				StreamIt.RemoveCurrent();
+			}
+			
+			for (const TPair<FObjectInStreamID, FConcertChangeStream_PutObject>& PutObjectPair : Request.ObjectsToPut)
+			{
+				if (PutObjectPair.Key.StreamId != StreamId)
+				{
+					continue;
+				}
+
+				const FSoftClassPath& PathToSet = PutObjectPair.Value.ClassPath;
+				const FConcertPropertySelection& SelectionToSet = PutObjectPair.Value.Properties;
+				FReplicatedObjectInfo& ObjectInfo = ReplicationMap.FindOrAdd(PutObjectPair.Key.Object);
+
+				// Write ClassPath if FindOrAdd just added it or if request specified a value to overwrite with ...
+				if (!ObjectInfo.ClassPath.IsValid() || PathToSet.IsValid())
+				{
+					ObjectInfo.ClassPath = PutObjectPair.Value.ClassPath;
+				}
+				// ... but one of these cases must occur
+				checkf(ObjectInfo.ClassPath.IsValid(), TEXT("Request did not validate ClassPath!"));
+
+				// Write Properties if FindOrAdd just added it or if request specified a value to overwrite with ...
+				if (ObjectInfo.PropertySelection.ReplicatedProperties.IsEmpty() || !SelectionToSet.ReplicatedProperties.IsEmpty())
+				{
+					ObjectInfo.PropertySelection = SelectionToSet;
+				}
+				// ... but one of these cases must occur
+				checkf(ObjectInfo.ClassPath.IsValid(), TEXT("Request did not validate Properties!"));
+			}
+		}
+		
+		for (const FReplicationStreamDescription_NetPacked& StreamDescription : Request.StreamsToAdd)
+		{
+			if (const TOptional<FReplicationStreamDescription> Unpacked = StreamDescription.Unpack()
+				; ensureMsgf(Unpacked, TEXT("Request did not validate unpacking!")))
+			{
+				StreamsToModify.Add(*Unpacked);
+			}
+		}
+	}
+}
diff --git a/Engine/Plugins/Developer/Concert/ConcertSync/ConcertSyncCore/Source/ConcertSyncCore/Public/Replication/ChangeStreamSharedUtils.h b/Engine/Plugins/Developer/Concert/ConcertSync/ConcertSyncCore/Source/ConcertSyncCore/Public/Replication/ChangeStreamSharedUtils.h
new file mode 100644
--- /dev/null
+++ b/Engine/Plugins/Developer/Concert/ConcertSync/ConcertSyncCore/Source/ConcertSyncCore/Public/Replication/ChangeStreamSharedUtils.h
@@ -0,0 +1,37 @@
+// Copyright Epic Games, Inc. All Rights Reserved.
+
+#pragma once
+
+#include "Containers/Array.h"
+#include "HAL/Platform.h"
+#include "Templates/Function.h"
+
+enum class EBreakBehavior : uint8;
+
+struct FConcertChangeStream_Request;
+struct FObjectInStreamID;
+struct FReplicationStreamDescription;
+
+/** This namespace contains the shared logic for applying FConcertChangeStream_Requests. */
+namespace UE::ConcertSyncCore::Replication::ChangeStreamUtils
+{
+	/**
+	 * Iterates through all objects that are removed and therefore the requesting client would lose authority over.
+	 * 
+	 * @param Request The request to parse
+	 * @param ExistingStreams The streams the request would be applied to
+	 * @param Callback The callback to process the objects
+	 */
+	CONCERTSYNCCORE_API void ForEachObjectLosingAuthority(
+		const FConcertChangeStream_Request& Request,
+		const TArray<FReplicationStreamDescription>& ExistingStreams,
+		TFunctionRef<EBreakBehavior(const FObjectInStreamID&)> Callback
+		);
+	
+	/**
+	 * Modifies RegisteredStreams as described by the request without validating that the Request is valid to apply.
+	 * @param Request The request to parse
+	 * @param StreamsToModify The streams to apply the request to
+	 */
+	CONCERTSYNCCORE_API void ApplyValidatedRequest(const FConcertChangeStream_Request& Request, IN OUT TArray<FReplicationStreamDescription>& StreamsToModify);
+}
diff --git a/Engine/Plugins/Developer/Concert/ConcertSync/ConcertSyncCore/Source/ConcertSyncCore/Public/Replication/Data/ReplicationStreamDescription.h b/Engine/Plugins/Developer/Concert/ConcertSync/ConcertSyncCore/Source/ConcertSyncCore/Public/Replication/Data/ReplicationStreamDescription.h
--- a/Engine/Plugins/Developer/Concert/ConcertSync/ConcertSyncCore/Source/ConcertSyncCore/Public/Replication/Data/ReplicationStreamDescription.h
+++ b/Engine/Plugins/Developer/Concert/ConcertSync/ConcertSyncCore/Source/ConcertSyncCore/Public/Replication/Data/ReplicationStreamDescription.h
@@ -23,6 +23,15 @@ struct FSharedReplicationStreamDescription
 	FObjectReplicationMap ReplicationMap;
 
 	// TODO: Add override settings to objects, such as override update frequency etc. Do it via TMap<FSoftObjectPath, FObjectReplicationSettings>.
+	
+	friend bool operator==(const FSharedReplicationStreamDescription& Left, const FSharedReplicationStreamDescription& Right)
+	{
+		return Left.Identifier == Right.Identifier && Left.ReplicationMap == Right.ReplicationMap;
+	}
+	friend bool operator!=(const FSharedReplicationStreamDescription& Left, const FSharedReplicationStreamDescription& Right)
+	{
+		return !(Left == Right);
+	}
 };
 
 /**
diff --git a/Engine/Plugins/Developer/Concert/ConcertSync/ConcertSyncCore/Source/ConcertSyncCore/Public/Replication/Messages/ConcertReplicationEvents.h b/Engine/Plugins/Developer/Concert/ConcertSync/ConcertSyncCore/Source/ConcertSyncCore/Public/Replication/Messages/ConcertReplicationEvents.h
--- a/Engine/Plugins/Developer/Concert/ConcertSync/ConcertSyncCore/Source/ConcertSyncCore/Public/Replication/Messages/ConcertReplicationEvents.h
+++ b/Engine/Plugins/Developer/Concert/ConcertSync/ConcertSyncCore/Source/ConcertSyncCore/Public/Replication/Messages/ConcertReplicationEvents.h
@@ -4,6 +4,7 @@
 
 #include "ConcertMessageData.h"
 #include "Replication/Data/ClientQueriedInfo.h"
+#include "Replication/Data/ObjectIds.h"
 #include "ConcertReplicationEvents.generated.h"
 
 /** Contains data to be applied to a replicated object. */
@@ -178,4 +179,130 @@ struct FConcertQueryReplicationInfo_Response
 	 */
 	UPROPERTY()
 	TMap<FGuid, FReplicationClientQueriedInfo> ClientInfo;
+};
+
+/** A request to add a new object to a stream or overwrite a pre-existing object's properties / class. */
+USTRUCT()
+struct FConcertChangeStream_PutObject
+{
+	GENERATED_BODY()
+
+	/**
+	 * The property selection the object should have.
+	 * Objects must always have a non-empty selection: use FConcertChangeStream_Request::ObjectsToRemove to remove objects.
+	 *
+	 * Any request that would leave an object's property selection empty will result in failure.
+	 */
+	UPROPERTY()
+	FConcertPropertySelection Properties;
+
+	/**
+	 * If Object is pre-existing, this property is optional. If different than None, this will change the class
+	 * If a new object is added, this property is mandatory. Not doing so will fail the request.
+	 *
+	 * Not specifying ClassPath nor Properties is an error: use FConcertChangeStream_Request::ObjectsToRemove to remove objects.
+	 */
+	UPROPERTY()
+	FSoftClassPath ClassPath;
+};
+
+/**
+ * Let's a client changing its streams owned on the server.
+ * 
+ * This request is processed atomically: it either succeeds completely or fails completely.
+ * If any sub-change causes a failure, the failure reason will be returned in the response and the client must make a new request.
+ */
+USTRUCT()
+struct FConcertChangeStream_Request
+{
+	GENERATED_BODY()
+
+	/**
+	 * Removes objects from pre-existing streams.
+	 * Supplying an object that is not in the specified stream does not cause failure (but is nonsensical).
+	 * 
+	 * If a stream has no registered objects after this operation, it is automatically removed.
+	 * If the requesting client has authority over these objects, authority is removed.
+	 */
+	UPROPERTY()
+	TSet<FObjectInStreamID> ObjectsToRemove;
+
+	/**
+	 * Adds new or modifies preexisting object definitions in pre-existing streams.
+	 *
+	 * If the requesting client and a different client have authority over the same object, you will get a conflict
+	 * if the different client already has authority over one of the properties you're adding here.
+	 * @see FConcertChangeStream_Response::AuthorityConflicts for some examples of conflicts.
+	 *
+	 * If the key identifies a stream that does not exist, the request will fail.
+	 */
+	UPROPERTY()
+	TMap<FObjectInStreamID, FConcertChangeStream_PutObject> ObjectsToPut;
+	
+	/**
+	 * New streams to add to the server.
+	 * Fails if any ID overlaps with a pre-existing one.
+	 */
+	UPROPERTY()
+	TArray<FReplicationStreamDescription_NetPacked> StreamsToAdd;
+
+	/**
+	 * Streams to remove from the server.
+	 * Supplying a stream that does not exist does not cause failure (but is nonsensical).
+	 * 
+	 * If the requesting client has authority over any of the objects contained in the stream, authority is removed.
+	 */
+	UPROPERTY()
+	TSet<FGuid> StreamsToRemove;
+};
+
+UENUM()
+enum class EConcertPutObjectErrorCode : uint8
+{
+	/** Stream that ObjectsToPut referenced was not registered on the server. */
+	UnresolvedStream,
+	/**
+	 * Either PutObject contained no data to update with (ensure either ClassPath or Properties is set),
+	 * or it tried to create a new object with insufficient data (make sure ClassPath and Properties are both specified).
+	 */
+	MissingData
+};
+
+/**
+ * Contains information about why a request failed. This info could be parsed and displayed to the end user as error.
+ * If there is even just one error, the entire requested has failed and no changes were made server-side.
+ */
+USTRUCT()
+struct FConcertChangeStream_Response
+{
+	GENERATED_BODY()
+
+	/**
+	 * Reports dynamic authority errors with ObjectsToPut:
+	 * Changing an object over which a client already has authority can yield unresolvable conflicts for which the entire FConcertChangeStream_Request is rejected.
+	 *
+	 * Let client R be the requester and client A be another client.
+	 * Example 1: Overlapping authority on different clients > Conflict
+	 * - R has authority over Foo's relative rotation property in stream S.
+	 * - A has authority over actor Foo's relative location property in some stream.
+	 * - R requests S to include Foo's relative location property: this conflicts with client A's authority.
+	 *
+	 * Example 2: Overlapping, properties on same client > No conflict
+	 * There is no conflict if another stream owned by the requester already has authority over a property you're adding to a different stream.
+	 * Streams of the same client may overlap properties. While this may not make much sense and actually be performance degrading, it would result in no logical errors so it is allowed.
+	 * Example: R has two streams S1 and S2. S1 contains the transform properties and S2 does not. Both S1 and S2 have authority over object Foo. It is legal to request S2 to contain the transform properties.
+	 */
+	UPROPERTY()
+	TMap<FObjectInStreamID, FReplicatedObjectId> AuthorityConflicts;
+
+	/** Reports semantic errors with ObjectsToPut. */
+	UPROPERTY()
+	TMap<FObjectInStreamID, EConcertPutObjectErrorCode> ObjectsToPutSemanticErrors;
+	
+	/** Streams that were in StreamsToAdd but that were not created. */
+	UPROPERTY()
+	TSet<FGuid> FailedStreamCreation;
+
+	bool IsSuccess() const { return AuthorityConflicts.IsEmpty() && ObjectsToPutSemanticErrors.IsEmpty() && FailedStreamCreation.IsEmpty(); }
+	bool IsFailure() const { return !IsSuccess(); }
 };
\ No newline at end of file
diff --git a/Engine/Plugins/Developer/Concert/ConcertSync/ConcertSyncServer/Source/ConcertSyncServer/Private/Replication/AuthorityManager.cpp b/Engine/Plugins/Developer/Concert/ConcertSync/ConcertSyncServer/Source/ConcertSyncServer/Private/Replication/AuthorityManager.cpp
--- a/Engine/Plugins/Developer/Concert/ConcertSync/ConcertSyncServer/Source/ConcertSyncServer/Private/Replication/AuthorityManager.cpp
+++ b/Engine/Plugins/Developer/Concert/ConcertSync/ConcertSyncServer/Source/ConcertSyncServer/Private/Replication/AuthorityManager.cpp
@@ -39,18 +39,91 @@ namespace UE::ConcertSyncServer::Replication
 		Session->UnregisterCustomRequestHandler<FConcertChangeAuthority_Request>();
 	}
 
-	bool FAuthorityManager::IsObjectChangeAllowed(const FReplicatedObjectId& ObjectChange) const
+	bool FAuthorityManager::HasAuthorityToChange(const FReplicatedObjectId& ObjectChange) const
 	{
 		const FClientAuthorityData* AuthorityData = ClientAuthorityData.Find(ObjectChange.SenderEndpointId);
 		const TSet<FSoftObjectPath>* OwnedObjects = AuthorityData ? AuthorityData->OwnedObjects.Find(ObjectChange.StreamId) : nullptr;
 		return OwnedObjects && OwnedObjects->Contains(ObjectChange.Object);
 	}
 
+	FAuthorityManager::EAuthorityResult FAuthorityManager::EnumerateAuthorityConflicts(
+		const FReplicatedObjectId& Object, 
+		const FConcertPropertySelection* OverwriteProperties,
+		FProcessAuthorityConflict ProcessConflict
+		) const
+	{
+		const FClientId& ClientId = Object.SenderEndpointId;
+		const FConcertPropertySelection* PropertiesToCheck = OverwriteProperties;
+		if (!PropertiesToCheck)
+		{
+			const FReplicationStreamDescription* Description = FindClientStreamById(ClientId, Object.StreamId);
+			const FReplicatedObjectInfo* PropertyInfo = Description ? Description->BaseDescription.ReplicationMap.ReplicatedObjects.Find(Object.Object) : nullptr;
+			PropertiesToCheck = PropertyInfo ? &PropertyInfo->PropertySelection : nullptr;
+		}
+
+		if (!PropertiesToCheck)
+		{
+			return EAuthorityResult::NoRegisteredProperties;
+		}
+
+		bool bFreeOfPropertyOverlaps = true;
+		const FClientId IgnoredClients[] = { ClientId };
+		ForEachClientWithPotentialConflict(
+			Object.Object,
+			[this, &ProcessConflict, PropertiesToCheck, &bFreeOfPropertyOverlaps](const FClientId& ClientId, const FStreamId& StreamId, const FConcertPropertySelection& WrittenProperties) mutable
+			{
+				// At this point we know that WritingClientId is sending WrittenProperties - if the properties overlap, the authority request is not possible
+				const bool bHasNoConflict = !WrittenProperties.OverlapsWith(*PropertiesToCheck);
+				bFreeOfPropertyOverlaps &= bHasNoConflict;
+				return bHasNoConflict
+					? EBreakBehavior::Continue
+					: ProcessConflict(ClientId, StreamId, WrittenProperties);
+			},
+			IgnoredClients
+			);
+		return bFreeOfPropertyOverlaps ? EAuthorityResult::Allowed : EAuthorityResult::Conflict;
+	}
+
+	bool FAuthorityManager::CanTakeAuthority(const FReplicatedObjectId& Object) const
+	{
+		return EnumerateAuthorityConflicts(Object) == EAuthorityResult::Allowed;
+	}
+
 	void FAuthorityManager::OnClientLeft(const FClientId& ClientEndpointId)
 	{
 		ClientAuthorityData.Remove(ClientEndpointId);
 	}
 
+	void FAuthorityManager::RemoveAuthority(const FReplicatedObjectId& Object)
+	{
+		const FClientId& ClientId = Object.SenderEndpointId;
+		FClientAuthorityData* ClientData = ClientAuthorityData.Find(ClientId);
+		if (!ClientData)
+		{
+			return;
+		}
+
+		const FStreamId& StreamId = Object.StreamId;
+		TSet<FSoftObjectPath>* AuthoredObjects = ClientData->OwnedObjects.Find(StreamId);
+		if (!AuthoredObjects)
+		{
+			return;
+		}
+
+		// This is all that is needed to remove authority.
+		AuthoredObjects->Remove(Object.Object);
+		
+		// Clean-up potentially empty entries
+		if (AuthoredObjects->IsEmpty())
+		{
+			ClientAuthorityData.Remove(StreamId);
+			if (ClientData->OwnedObjects.IsEmpty())
+			{
+				ClientAuthorityData.Remove(ClientId);
+			}
+		}
+	}
+
 	EConcertSessionResponseCode FAuthorityManager::HandleChangeAuthorityRequest(
 		const FConcertSessionContext& ConcertSessionContext,
 		const FConcertChangeAuthority_Request& Request,
@@ -64,7 +137,8 @@ namespace UE::ConcertSyncServer::Replication
 		const FClientId& ClientId = ConcertSessionContext.SourceEndpointId;
 		Private::ForEachReplicatedObject(Request.TakeAuthority, [this, &Response, &AuthorityData, &ClientId](const FStreamId& StreamId, const FSoftObjectPath& ObjectPath)
 		{
-			if (CanTakeAuthority(AuthorityData, ClientId, { StreamId, ObjectPath }))
+			const FReplicatedObjectId ObjectToAuthor{ { StreamId, ObjectPath }, ClientId };
+			if (CanTakeAuthority(ObjectToAuthor))
 			{
 				UE_LOG(LogConcert, Log, TEXT("Transferred authority of %s to client %s for their stream %s"), *ObjectPath.ToString(), *ClientId.ToString(EGuidFormats::Short), *StreamId.ToString(EGuidFormats::Short));
 				AuthorityData.OwnedObjects.FindOrAdd(StreamId).Add(ObjectPath);
@@ -78,7 +152,8 @@ namespace UE::ConcertSyncServer::Replication
 		
 		Private::ForEachReplicatedObject(Request.TakeAuthority, [this, &Response, &AuthorityData, &ClientId](const FStreamId& StreamId, const FSoftObjectPath& ObjectPath)
 		{
-			if (CanTakeAuthority(AuthorityData, ClientId, { StreamId, ObjectPath }))
+			const FReplicatedObjectId ObjectToAuthor{ { StreamId, ObjectPath }, ClientId };
+			if (CanTakeAuthority(ObjectToAuthor))
 			{
 				UE_LOG(LogConcert, Log, TEXT("Transferred authority of %s to client %s for their stream %s"), *ObjectPath.ToString(), *ClientId.ToString(EGuidFormats::Short), *StreamId.ToString(EGuidFormats::Short));
 				AuthorityData.OwnedObjects.FindOrAdd(StreamId).Add(ObjectPath);
@@ -90,7 +165,7 @@ namespace UE::ConcertSyncServer::Replication
 			}
 		});
 		
-		Private::ForEachReplicatedObject(Request.ReleaseAuthority, [this, &Response, &AuthorityData, &ClientId](const FStreamId& StreamId, const FSoftObjectPath& ObjectPath)
+		Private::ForEachReplicatedObject(Request.ReleaseAuthority, [this, &AuthorityData](const FStreamId& StreamId, const FSoftObjectPath& ObjectPath)
 		{
 			TSet<FSoftObjectPath>* OwnedObjects = AuthorityData.OwnedObjects.Find(StreamId);
 			// Though dubious, it is a valid request for the client to release non-owned objects
@@ -109,38 +184,6 @@ namespace UE::ConcertSyncServer::Replication
 
 		return EConcertSessionResponseCode::Success;
 	}
-
-	bool FAuthorityManager::CanTakeAuthority(const FClientAuthorityData& ClientData, const FClientId& ClientId, const FObjectInStreamID& Object) const
-	{
-		/*
-		 * This function is unoptimized: it repeatedly iterates through every client's registered streams and properties.
-		 * However it should be fine since:
-		 * 1. Changing authority does not happen too often though,
-		 * 2. There will be a relatively small number of clients, streams, objects & properties.
-		 */
-		
-		const FReplicationStreamDescription* Description = FindClientStreamById(ClientId, Object.StreamId);
-		const FSoftObjectPath& ObjectPath = Object.Object;
-		const FReplicatedObjectInfo* PropertyInfo = Description ? Description->BaseDescription.ReplicationMap.ReplicatedObjects.Find(ObjectPath) : nullptr;
-		if (!PropertyInfo)
-		{
-			return false;
-		}
-
-		bool bFreeOfPropertyOverlaps = true;
-		const FClientId IgnoredClients[] = { ClientId };
-		ForEachClientWithPotentialConflict(ObjectPath, [this, &ClientId, &PropertyInfo, &bFreeOfPropertyOverlaps](const FClientId& WritingClientId, const FConcertPropertySelection& WrittenProperties) mutable
-		{
-			// At this point we know that WritingClientId is sending WrittenProperties - if the properties overlap, the authority request is not possible
-			bFreeOfPropertyOverlaps &= !WrittenProperties.OverlapsWith(PropertyInfo->PropertySelection);
-			if (bFreeOfPropertyOverlaps)
-			{
-				return EBreakBehavior::Continue;
-			}
-			return EBreakBehavior::Break;
-		}, IgnoredClients);
-		return bFreeOfPropertyOverlaps;
-	}
 	
 	const FReplicationStreamDescription* FAuthorityManager::FindClientStreamById(const FClientId& ClientId, const FStreamId& StreamId) const
 	{
@@ -159,7 +202,7 @@ namespace UE::ConcertSyncServer::Replication
 
 	void FAuthorityManager::ForEachClientWithPotentialConflict(
 		const FSoftObjectPath& Object,
-		TFunctionRef<EBreakBehavior(const FClientId& ClientId, const FConcertPropertySelection& WrittenProperties)> Callback,
+		FProcessAuthorityConflict Callback,
 		TArrayView<const FClientId> IgnoredClients
 		) const
 	{
@@ -181,7 +224,7 @@ namespace UE::ConcertSyncServer::Replication
 			Getters.ForEachStream(ClientEndpointId, [&Object, &Callback, &ClientEndpointId, OtherClientData, &Result](const FReplicationStreamDescription& Stream) mutable
 			{
 				// If client has not claimed authority over this object in this stream, skip
-				const FGuid& StreamId = Stream.BaseDescription.Identifier;
+				const FStreamId& StreamId = Stream.BaseDescription.Identifier;
 				const TSet<FSoftObjectPath>* ClientControlledObjects = OtherClientData->OwnedObjects.Find(StreamId);
 				const bool bClientHasSomeAuthorityOverObject = ClientControlledObjects && ClientControlledObjects->Contains(Object);
 				if (!bClientHasSomeAuthorityOverObject)
@@ -192,7 +235,7 @@ namespace UE::ConcertSyncServer::Replication
 				// This client is using the request object: report the potential conflict ...
 				const FObjectReplicationMap& ObjectReplicationMap = Stream.BaseDescription.ReplicationMap;
 				const FReplicatedObjectInfo* ReplicationObjectInfo = ObjectReplicationMap.ReplicatedObjects.Find(Object);
-				if (ReplicationObjectInfo && Callback(ClientEndpointId, ReplicationObjectInfo->PropertySelection) == EBreakBehavior::Break)
+				if (ReplicationObjectInfo && Callback(ClientEndpointId, StreamId, ReplicationObjectInfo->PropertySelection) == EBreakBehavior::Break)
 				{
 					// ... conflict ends iteration
 					Result = EBreakBehavior::Break;
diff --git a/Engine/Plugins/Developer/Concert/ConcertSync/ConcertSyncServer/Source/ConcertSyncServer/Private/Replication/AuthorityManager.h b/Engine/Plugins/Developer/Concert/ConcertSync/ConcertSyncServer/Source/ConcertSyncServer/Private/Replication/AuthorityManager.h
--- a/Engine/Plugins/Developer/Concert/ConcertSync/ConcertSyncServer/Source/ConcertSyncServer/Private/Replication/AuthorityManager.h
+++ b/Engine/Plugins/Developer/Concert/ConcertSync/ConcertSyncServer/Source/ConcertSyncServer/Private/Replication/AuthorityManager.h
@@ -4,6 +4,7 @@
 
 #include "ConcertMessages.h"
 #include "Misc/EBreakBehavior.h"
+#include "Replication/Data/ObjectIds.h"
 #include "Templates/Function.h"
 
 class IConcertSession;
@@ -40,7 +41,9 @@ namespace UE::ConcertSyncServer::Replication
 	{
 	public:
 		
+		using FStreamId = FGuid;
 		using FClientId = FGuid;
+		using FProcessAuthorityConflict = TFunctionRef<EBreakBehavior(const FClientId& ClientId, const FStreamId& StreamId, const FConcertPropertySelection& WrittenProperties)>;
 
 		FAuthorityManager(IAuthorityManagerGetters& Getters, TSharedRef<IConcertSession> InSession);
 		~FAuthorityManager();
@@ -49,10 +52,34 @@ namespace UE::ConcertSyncServer::Replication
 		 * Checks whether the client that sent the identified object had authority to send it.
 		 * @return Whether the server should process the object change.
 		 */
-		bool IsObjectChangeAllowed(const FReplicatedObjectId& ObjectChange) const;
+		bool HasAuthorityToChange(const FReplicatedObjectId& ObjectChange) const;
+
+		enum class EAuthorityResult : uint8
+		{
+			/** The client is allowed to take authority */
+			Allowed,
+			/** There was at least one conflict */
+			Conflict,
+			/** No conflicts were checked because OverwriteProperties was nullptr and the client had not registered any properties to send for the given object. */
+			NoRegisteredProperties
+		};
+		/**
+		 * Enumerates all authority conflicts, if any, that would occur if ObjectChange.SenderEndpointId were to take authority over the identified object.
+		 * You can optionally supply OverwriteProperties, which is useful if you're about to change the contents of the stream.
+		 * @return Whether there were any conflicts
+		 */
+		EAuthorityResult EnumerateAuthorityConflicts(
+			const FReplicatedObjectId& Object,
+			const FConcertPropertySelection* OverwriteProperties = nullptr,
+			FProcessAuthorityConflict ProcessConflict = [](auto&, auto&, auto&){ return EBreakBehavior::Break; }
+			) const;
+		/** Whether it is legal for this the client identified by ClientId to take control over the object given the stream the client has registered. */
+		bool CanTakeAuthority(const FReplicatedObjectId& Object) const;
 
 		/** Notifies this manager that the client has left, which means all their authority is now gone. */
 		void OnClientLeft(const FClientId& ClientEndpointId);
+		/** Takes away authority from the given client from the given object. */
+		void RemoveAuthority(const FReplicatedObjectId& Object);
 
 	private:
 
@@ -61,7 +88,6 @@ namespace UE::ConcertSyncServer::Replication
 		/** The session under which this manager operates. */
 		TSharedRef<IConcertSession> Session;
 
-		using FStreamId = FGuid;
 		struct FClientAuthorityData
 		{
 			/** Objects the client has authority over. */
@@ -75,16 +101,13 @@ namespace UE::ConcertSyncServer::Replication
 			FConcertChangeAuthority_Response& Response
 			);
 
-		/** Whether it is legal for this the client identified by ClientId to take control over the object given the stream the client has registered. */
-		bool CanTakeAuthority(const FClientAuthorityData& ClientData, const FClientId& ClientId, const FObjectInStreamID& Object) const;
-
 		/** Finds a stream registered with the client by its ID. */
 		const FReplicationStreamDescription* FindClientStreamById(const FClientId& ClientId, const FStreamId& StreamId) const;
 
 		/** Iterates through all clients that are already writing to Object. */
 		void ForEachClientWithPotentialConflict(
 			const FSoftObjectPath& Object,
-			TFunctionRef<EBreakBehavior(const FClientId& ClientId, const FConcertPropertySelection& WrittenProperties)> Callback,
+			FProcessAuthorityConflict Callback,
 			TArrayView<const FClientId> IgnoredClients = {}
 			) const;
 	};
diff --git a/Engine/Plugins/Developer/Concert/ConcertSync/ConcertSyncServer/Source/ConcertSyncServer/Private/Replication/ConcertReplicationClient.cpp b/Engine/Plugins/Developer/Concert/ConcertSync/ConcertSyncServer/Source/ConcertSyncServer/Private/Replication/ConcertReplicationClient.cpp
--- a/Engine/Plugins/Developer/Concert/ConcertSync/ConcertSyncServer/Source/ConcertSyncServer/Private/Replication/ConcertReplicationClient.cpp
+++ b/Engine/Plugins/Developer/Concert/ConcertSync/ConcertSyncServer/Source/ConcertSyncServer/Private/Replication/ConcertReplicationClient.cpp
@@ -3,6 +3,7 @@
 #include "ConcertReplicationClient.h"
 
 #include "Processing/ServerReplicationDataQueuer.h"
+#include "Replication/ChangeStreamSharedUtils.h"
 #include "Replication/Messages/ConcertReplicationHandshakeMessages.h"
 
 namespace UE::ConcertSyncServer::Replication
@@ -23,4 +24,10 @@ namespace UE::ConcertSyncServer::Replication
 	{
 		DataRelay.ProcessObjects(TimeBudget);
 	}
+
+	void FConcertReplicationClient::ApplyValidatedRequest(const FConcertChangeStream_Request& Request)
+	{
+		// Right now there is nothing further to do but if in future you need to update some client systems of the change, this is the place to do it.
+		ConcertSyncCore::Replication::ChangeStreamUtils::ApplyValidatedRequest(Request, StreamDescriptions);
+	}
 }
diff --git a/Engine/Plugins/Developer/Concert/ConcertSync/ConcertSyncServer/Source/ConcertSyncServer/Private/Replication/ConcertReplicationClient.h b/Engine/Plugins/Developer/Concert/ConcertSync/ConcertSyncServer/Source/ConcertSyncServer/Private/Replication/ConcertReplicationClient.h
--- a/Engine/Plugins/Developer/Concert/ConcertSync/ConcertSyncServer/Source/ConcertSyncServer/Private/Replication/ConcertReplicationClient.h
+++ b/Engine/Plugins/Developer/Concert/ConcertSync/ConcertSyncServer/Source/ConcertSyncServer/Private/Replication/ConcertReplicationClient.h
@@ -8,6 +8,7 @@
 
 struct FReplicationStreamDescription;
 struct FConcertReplication_Join_Request;
+struct FConcertChangeStream_Request;
 
 namespace UE::ConcertSyncCore
 {
@@ -35,16 +36,18 @@ namespace UE::ConcertSyncServer::Replication
 		 * Process given a time budget. The time budget may be exceeded but we'll try not to and to stay as close to the budget as possible.
 		 */
 		void ProcessClient(float TimeBudget);
+		
+		/** Updates the StreamDescriptions array with the changes from Request. The request already passed validation and is valid to apply. */
+		void ApplyValidatedRequest(const FConcertChangeStream_Request& Request);
 
 		const FGuid& GetClientEndpointId() const { return ClientEndpointId; }
-
 		const TArray<FReplicationStreamDescription>& GetStreamDescriptions() const { return StreamDescriptions; }
 
 	private:
 
 		/** The streams this client offered to send. */
-		const TArray<FReplicationStreamDescription> StreamDescriptions;
-
+		TArray<FReplicationStreamDescription> StreamDescriptions;
+		
 		/** This client's endpoint ID. */
 		const FGuid ClientEndpointId;
 		
diff --git a/Engine/Plugins/Developer/Concert/ConcertSync/ConcertSyncServer/Source/ConcertSyncServer/Private/Replication/ConcertServerReplicationManager.cpp b/Engine/Plugins/Developer/Concert/ConcertSync/ConcertSyncServer/Source/ConcertSyncServer/Private/Replication/ConcertServerReplicationManager.cpp
--- a/Engine/Plugins/Developer/Concert/ConcertSync/ConcertSyncServer/Source/ConcertSyncServer/Private/Replication/ConcertServerReplicationManager.cpp
+++ b/Engine/Plugins/Developer/Concert/ConcertSync/ConcertSyncServer/Source/ConcertSyncServer/Private/Replication/ConcertServerReplicationManager.cpp
@@ -22,6 +22,7 @@ namespace UE::ConcertSyncServer::Replication
 	{
 		Session->RegisterCustomRequestHandler<FConcertReplication_Join_Request, FConcertReplication_Join_Response>(this, &FConcertServerReplicationManager::HandleJoinReplicationSessionRequest);
 		Session->RegisterCustomRequestHandler<FConcertQueryReplicationInfo_Request, FConcertQueryReplicationInfo_Response>(this, &FConcertServerReplicationManager::HandleQueryReplicationInfoRequest);
+		Session->RegisterCustomRequestHandler<FConcertChangeStream_Request, FConcertChangeStream_Response>(this, &FConcertServerReplicationManager::HandleChangeStreamRequest);
 		Session->RegisterCustomEventHandler<FConcertReplication_LeaveEvent>(this, &FConcertServerReplicationManager::HandleLeaveReplicationSessionRequest);
 		Session->OnSessionClientChanged().AddRaw(this, &FConcertServerReplicationManager::OnConnectionChanged);
 
@@ -187,7 +188,7 @@ namespace UE::ConcertSyncServer::Replication
 				ObjectInfo.Object = Pair.Key;
 				ObjectInfo.StreamId = StreamId;
 				
-				if (AuthorityManager->IsObjectChangeAllowed(ObjectInfo))
+				if (AuthorityManager->HasAuthorityToChange(ObjectInfo))
 				{
 					Info.AuthoredObjects.Add(Pair.Key);
 				}
@@ -204,8 +205,7 @@ namespace UE::ConcertSyncServer::Replication
 
 	void FConcertServerReplicationManager::HandleLeaveReplicationSessionRequest(
 		const FConcertSessionContext& ConcertSessionContext,
-		const FConcertReplication_LeaveEvent& EventData
-		)
+		const FConcertReplication_LeaveEvent& EventData)
 	{
 		const FGuid ClientEndpointId = ConcertSessionContext.SourceEndpointId;
 		UE_LOG(LogConcert, Log, TEXT("Received replication leave request from endpoint %s"), *ClientEndpointId.ToString());
diff --git a/Engine/Plugins/Developer/Concert/ConcertSync/ConcertSyncServer/Source/ConcertSyncServer/Private/Replication/ConcertServerReplicationManager.h b/Engine/Plugins/Developer/Concert/ConcertSync/ConcertSyncServer/Source/ConcertSyncServer/Private/Replication/ConcertServerReplicationManager.h
--- a/Engine/Plugins/Developer/Concert/ConcertSync/ConcertSyncServer/Source/ConcertSyncServer/Private/Replication/ConcertServerReplicationManager.h
+++ b/Engine/Plugins/Developer/Concert/ConcertSync/ConcertSyncServer/Source/ConcertSyncServer/Private/Replication/ConcertServerReplicationManager.h
@@ -38,6 +38,7 @@ namespace UE::ConcertSyncServer::Replication
 		, public FNoncopyable
 	{
 	public:
+
 		explicit FConcertServerReplicationManager(TSharedRef<IConcertServerSession> InLiveSession);
 		virtual ~FConcertServerReplicationManager() override;
 
@@ -65,21 +66,24 @@ namespace UE::ConcertSyncServer::Replication
 		/** Clients that have requested to join replication. Maps client ID to replication info. */
 		TMap<FGuid, TSharedRef<FConcertReplicationClient>> Clients;
 
-		// Event handlers
+		// Joining
 		EConcertSessionResponseCode HandleJoinReplicationSessionRequest(const FConcertSessionContext& ConcertSessionContext, const FConcertReplication_Join_Request& Request, FConcertReplication_Join_Response& Response);
 		EConcertSessionResponseCode InternalHandleJoinReplicationSessionRequest(const FConcertSessionContext& ConcertSessionContext, const FConcertReplication_Join_Request& Request, FConcertReplication_Join_Response& Response);
 
+		// Querying
 		EConcertSessionResponseCode HandleQueryReplicationInfoRequest(const FConcertSessionContext& ConcertSessionContext, const FConcertQueryReplicationInfo_Request& Request, FConcertQueryReplicationInfo_Response& Response);
 		/** Gets all registered streams and optionally removes the properties. */
 		TArray<FSharedReplicationStreamDescription> BuildClientStreamInfo(const FConcertReplicationClient& Client, bool bSkipProperties) const;
 		/** Maps the client's streams to the objects in that stream the client has taken authority over. */
 		TArray<FReplicationAuthorityInfo> BuildClientAuthorityInfo(const FConcertReplicationClient& Client) const;
-		
+
+		// Changing streams
+		EConcertSessionResponseCode HandleChangeStreamRequest(const FConcertSessionContext& ConcertSessionContext, const FConcertChangeStream_Request& Request, FConcertChangeStream_Response& Response);
+
+		// Leaving
 		void HandleLeaveReplicationSessionRequest(const FConcertSessionContext& ConcertSessionContext, const FConcertReplication_LeaveEvent& EventData);
-		
 		void OnConnectionChanged(IConcertServerSession& ConcertServerSession, EConcertClientStatus ConcertClientStatus, const FConcertSessionClientInfo& ConcertSessionClientInfo);
 
-	private:
 		/**
 		 * Ticks all clients which causes clients to process pending data and send it to the corresponding endpoints.
 		 * 
diff --git a/Engine/Plugins/Developer/Concert/ConcertSync/ConcertSyncServer/Source/ConcertSyncServer/Private/Replication/ConcertServerReplicationManager_ChangeStream.cpp b/Engine/Plugins/Developer/Concert/ConcertSync/ConcertSyncServer/Source/ConcertSyncServer/Private/Replication/ConcertServerReplicationManager_ChangeStream.cpp
new file mode 100644
--- /dev/null
+++ b/Engine/Plugins/Developer/Concert/ConcertSync/ConcertSyncServer/Source/ConcertSyncServer/Private/Replication/ConcertServerReplicationManager_ChangeStream.cpp
@@ -0,0 +1,168 @@
+// Copyright Epic Games, Inc. All Rights Reserved.
+
+#include "ConcertServerReplicationManager.h"
+
+#include "AuthorityManager.h"
+#include "ConcertLogGlobal.h"
+#include "IConcertSession.h"
+#include "Replication/ChangeStreamSharedUtils.h"
+#include "Replication/ConcertReplicationClient.h"
+#include "Replication/Messages/ConcertReplicationHandshakeMessages.h"
+#include "Replication/Processing/ObjectReplicationCache.h"
+
+namespace UE::ConcertSyncServer::Replication
+{
+	namespace Private
+	{
+		static const FReplicationStreamDescription* FindExistingStream(const FConcertReplicationClient& Client, const FGuid& StreamId)
+		{
+			return Client.GetStreamDescriptions().FindByPredicate([&StreamId](const FReplicationStreamDescription& Description)
+			{
+				return Description.BaseDescription.Identifier == StreamId;
+			});
+		}
+		
+		/** Validates that ObjectsToPut writes only to pre-existing streams. */
+		static void ValidatePutObjectsRequestSemantics(const FConcertChangeStream_Request& Request, const FConcertReplicationClient& Client, FConcertChangeStream_Response& OutResponse)
+		{
+			const auto PutObjectHasEnoughData = [](const FReplicationStreamDescription* ExistingStream, const FSoftObjectPath& ChangedObjectPath, const FConcertChangeStream_PutObject& PutObject)
+			{
+				const bool bHasProperties = !PutObject.Properties.ReplicatedProperties.IsEmpty();
+				const bool bHasClassPath = !PutObject.ClassPath.IsNull();
+				const bool bIsEditingExistingObjectDefinition = ExistingStream->BaseDescription.ReplicationMap.ReplicatedObjects.Contains(ChangedObjectPath);
+				const bool bHasEnoughData = (!bIsEditingExistingObjectDefinition && bHasProperties && bHasClassPath)
+					|| (bIsEditingExistingObjectDefinition && bHasProperties)
+					|| (bIsEditingExistingObjectDefinition && bHasClassPath);
+				return bHasEnoughData;
+			};
+			
+			for (const TPair<FObjectInStreamID, FConcertChangeStream_PutObject>& Change : Request.ObjectsToPut)
+			{
+				const FObjectInStreamID ChangedObject = Change.Key;
+				
+				const FGuid& StreamToModify = ChangedObject.StreamId;
+				const FReplicationStreamDescription* ExistingStream = FindExistingStream(Client, StreamToModify);
+				const bool bStreamExists = ExistingStream != nullptr;
+				if (bStreamExists)
+				{
+					// FReplicatedObjectInfo must have non-empty values for its members. Hence we must check that
+					// 1. if creating a new entry, both fields are given
+					// 2. if writing to pre-existing entry, at least one field is given
+					const FConcertChangeStream_PutObject& PutObject = Change.Value; 
+					if (!PutObjectHasEnoughData(ExistingStream, ChangedObject.Object, PutObject))
+					{
+						OutResponse.ObjectsToPutSemanticErrors.Add(ChangedObject, EConcertPutObjectErrorCode::MissingData);
+					}
+				}
+				else
+				{
+					// ObjectsToPut can only write to pre-existing streams
+					OutResponse.ObjectsToPutSemanticErrors.Add(ChangedObject, EConcertPutObjectErrorCode::UnresolvedStream);
+				}
+			}
+		}
+
+		/** Checks that StreamsToAdd do not conflict with pre-existing ones and that all IDs in the request are also unique. */
+		static void ValidateAddedStreamsAreUnique(const FConcertChangeStream_Request& Request, const FConcertReplicationClient& Client, FConcertChangeStream_Response& OutResponse)
+		{
+			TSet<FGuid> DuplicateEntryDetection;
+			for (const FReplicationStreamDescription_NetPacked& NewStream : Request.StreamsToAdd)
+			{
+				// StreamsToAdd is invalid if there is already a stream with the same ID registered ...
+				const FGuid& NewStreamId = NewStream.BaseDescription.Identifier;
+				const bool bIdAlreadyExists = FindExistingStream(Client, NewStreamId) != nullptr;
+
+				// ... or StreamsToAdd contains the same ID multiple times
+				bool bIsDuplicateEntry = false;
+				DuplicateEntryDetection.FindOrAdd(NewStreamId, &bIsDuplicateEntry);
+				
+				if (bIdAlreadyExists || bIsDuplicateEntry)
+				{
+					OutResponse.FailedStreamCreation.Add(NewStreamId);
+				}
+			}
+		}
+
+		/** If the requesting client has authority over a changed object, checks that no other client has authority over the properties being added. */
+		static void LookForAuthorityConflicts(const FConcertChangeStream_Request& Request, const FConcertReplicationClient& Client, const FAuthorityManager& AuthorityManager, FConcertChangeStream_Response& OutResponse)
+		{
+			const FGuid ClientEndpointId = Client.GetClientEndpointId();
+			for (const TPair<FObjectInStreamID, FConcertChangeStream_PutObject>& PutObjectPair : Request.ObjectsToPut)
+			{
+				// No conflict possible if requesting client does not have authority over the changed object
+				const FReplicatedObjectId ReplicatedObjectInfo { { PutObjectPair.Key }, ClientEndpointId };
+				if (!AuthorityManager.HasAuthorityToChange(ReplicatedObjectInfo))
+				{
+					continue;
+				}
+
+				// Requester not changing properties (must be changing ClassPath)? Also no conflict possible.
+				const FConcertPropertySelection& PropertySelection = PutObjectPair.Value.Properties;
+				const bool bIsEditingProperties = !PropertySelection.ReplicatedProperties.IsEmpty();
+				if (!bIsEditingProperties)
+				{
+					continue;
+				}
+
+				// Simply check whether any other client is already sending any of the requested properties.
+				AuthorityManager.EnumerateAuthorityConflicts(ReplicatedObjectInfo, &PropertySelection,
+					[&OutResponse, &ReplicatedObjectInfo](const FGuid& ClientId, const FGuid& StreamId, const FConcertPropertySelection& WrittenProperties)
+					{
+						const FReplicatedObjectId ConflictingObject = { { StreamId, ReplicatedObjectInfo.Object }, ClientId };
+						OutResponse.AuthorityConflicts.Add(ReplicatedObjectInfo, ConflictingObject);
+						return EBreakBehavior::Continue;
+					});
+			}
+		}
+
+		/** Stream description can have UObjects attached which fail to serialize on the receiving end. If so, reject. */
+		static void ValidateUnpacking(const FConcertChangeStream_Request& Request, FConcertChangeStream_Response& OutResponse)
+		{
+			for (const FReplicationStreamDescription_NetPacked& ToUnpack : Request.StreamsToAdd)
+			{
+				if (!ToUnpack.Unpack())
+				{
+					OutResponse.FailedStreamCreation.Add(ToUnpack.BaseDescription.Identifier);
+				}
+			}
+		}
+
+		/** Checks whether this request is valid to apply. */
+		static bool ShouldAcceptRequest(const FConcertChangeStream_Request& Request, const FConcertReplicationClient& Client, const FAuthorityManager& AuthorityManager, FConcertChangeStream_Response& OutResponse)
+		{
+			ValidatePutObjectsRequestSemantics(Request, Client, OutResponse);
+			ValidateAddedStreamsAreUnique(Request, Client, OutResponse);
+			LookForAuthorityConflicts(Request, Client, AuthorityManager, OutResponse);
+			ValidateUnpacking(Request, OutResponse);
+			return OutResponse.IsSuccess();
+		}
+	}
+
+	EConcertSessionResponseCode FConcertServerReplicationManager::HandleChangeStreamRequest(
+		const FConcertSessionContext& ConcertSessionContext,
+		const FConcertChangeStream_Request& Request,
+		FConcertChangeStream_Response& Response)
+	{
+		const FGuid SendingClientId = ConcertSessionContext.SourceEndpointId;
+		const TSharedRef<FConcertReplicationClient>* SendingClient = Clients.Find(SendingClientId);
+		if (SendingClient && Private::ShouldAcceptRequest(Request, SendingClient->Get(), AuthorityManager.Get(), Response))
+		{
+			// If the client had authority over any objects that were removed by this request, authority must be cleaned up
+			ConcertSyncCore::Replication::ChangeStreamUtils::ForEachObjectLosingAuthority(Request, SendingClient->Get().GetStreamDescriptions(),
+				[this, &SendingClientId](const FObjectInStreamID& RemovedObject)
+				{
+					AuthorityManager->RemoveAuthority({ RemovedObject, SendingClientId});
+					return EBreakBehavior::Continue;
+				});
+			
+			SendingClient->Get().ApplyValidatedRequest(Request);
+		}
+		else
+		{
+			UE_LOG(LogConcert, Warning, TEXT("Rejecting ChangeStream request from %s"), *SendingClientId.ToString(EGuidFormats::Short));
+		}
+		
+		return EConcertSessionResponseCode::Success;
+	}
+}
+
diff --git a/Engine/Plugins/Developer/Concert/ConcertSync/ConcertSyncServer/Source/ConcertSyncServer/Private/Replication/Processing/ServerObjectReplicationReceiver.cpp b/Engine/Plugins/Developer/Concert/ConcertSync/ConcertSyncServer/Source/ConcertSyncServer/Private/Replication/Processing/ServerObjectReplicationReceiver.cpp
--- a/Engine/Plugins/Developer/Concert/ConcertSync/ConcertSyncServer/Source/ConcertSyncServer/Private/Replication/Processing/ServerObjectReplicationReceiver.cpp
+++ b/Engine/Plugins/Developer/Concert/ConcertSync/ConcertSyncServer/Source/ConcertSyncServer/Private/Replication/Processing/ServerObjectReplicationReceiver.cpp
@@ -25,6 +25,6 @@ namespace UE::ConcertSyncServer::Replication
 		) const
 	{
 		const FReplicatedObjectId ReplicatedObjectInfo { { StreamEvent.StreamId, ObjectEvent.ReplicatedObject }, SessionContext.SourceEndpointId };
-		return AuthorityManager->IsObjectChangeAllowed(ReplicatedObjectInfo);
+		return AuthorityManager->HasAuthorityToChange(ReplicatedObjectInfo);
 	}
 }
diff --git a/Engine/Plugins/Developer/Concert/ConcertSync/ConcertSyncServer/Source/ConcertSyncServer/Private/Replication/Util/JoinRequestValidation.h b/Engine/Plugins/Developer/Concert/ConcertSync/ConcertSyncServer/Source/ConcertSyncServer/Private/Replication/Util/JoinRequestValidation.h
--- a/Engine/Plugins/Developer/Concert/ConcertSync/ConcertSyncServer/Source/ConcertSyncServer/Private/Replication/Util/JoinRequestValidation.h
+++ b/Engine/Plugins/Developer/Concert/ConcertSync/ConcertSyncServer/Source/ConcertSyncServer/Private/Replication/Util/JoinRequestValidation.h
@@ -3,7 +3,6 @@
 #pragma once
 
 #include "Containers/Array.h"
-#include "Templates/SharedPointer.h"
 #include "Templates/Tuple.h"
 
 enum class EJoinReplicationErrorCode : uint8;
@@ -12,10 +11,6 @@ struct FReplicationStreamDescription;
 
 namespace UE::ConcertSyncServer::Replication
 {
-	class FConcertReplicationClient;
-
-	using FGetClientName = TFunctionRef<FString(const FGuid& ClientEndpointID)>;
-	
 	/**
 	 * Attempts to unpack the request and validates it.
 	 * Constraints:
diff --git a/Engine/Plugins/Developer/Concert/ConcertSync/ConcertSyncTest/Source/ConcertSyncTest/Private/Replication/ReplicationStreamRequestTests.cpp b/Engine/Plugins/Developer/Concert/ConcertSync/ConcertSyncTest/Source/ConcertSyncTest/Private/Replication/ReplicationStreamRequestTests.cpp
--- a/Engine/Plugins/Developer/Concert/ConcertSync/ConcertSyncTest/Source/ConcertSyncTest/Private/Replication/ReplicationStreamRequestTests.cpp
+++ b/Engine/Plugins/Developer/Concert/ConcertSync/ConcertSyncTest/Source/ConcertSyncTest/Private/Replication/ReplicationStreamRequestTests.cpp
@@ -5,12 +5,12 @@
 #include "Replication/Data/ReplicationStreamDescription.h"
 #include "Replication/IConcertClientReplicationManager.h"
 #include "TestReflectionObject.h"
+#include "Util/SendReceiveGenericStreamTestBase.h"
 #include "Util/SendReceiveObjectTestBase.h"
 
-#include "Components/StaticMeshComponent.h"
 #include "Misc/AutomationTest.h"
-#include "Misc/Paths.h"
-#include "HAL/FileManager.h"
+#include "UObject/UObjectGlobals.h"
+#include "UObject/Package.h"
 
 namespace UE::ConcertSyncTests::Replication::Stream
 {
@@ -20,7 +20,7 @@ namespace UE::ConcertSyncTests::Replication::Stream
 	 * - that the authority updates,
 	 * - the EConcertQueryClientStreamFlags flags
 	 */
-	IMPLEMENT_CUSTOM_SIMPLE_AUTOMATION_TEST(FQueryOtherClientStreams, FSendReceiveObjectTestBase, "Concert.Replication.QueryClientStream", EAutomationTestFlags::EditorContext | EAutomationTestFlags::EngineFilter);
+	IMPLEMENT_CUSTOM_SIMPLE_AUTOMATION_TEST(FQueryOtherClientStreams, FSendReceiveObjectTestBase, "Concert.Replication.Stream.QueryClientStream", EAutomationTestFlags::EditorContext | EAutomationTestFlags::EngineFilter);
 	bool FQueryOtherClientStreams::RunTest(const FString& Parameters)
 	{
 		// 1. Init
@@ -76,6 +76,7 @@ namespace UE::ConcertSyncTests::Replication::Stream
 				bSenderReceivedResponse = true;
 				TestEqual(TEXT("No rejection taking authority"), Response.RejectedObjects.Num(), 0); 
 			});
+		TestTrue(TEXT("Received response taking authority"), bSenderReceivedResponse);
 
 		bool bReceivedSecondQueryResponse = false;
 		ClientReplicationManager_Receiver->QueryClientInfo(Request)
@@ -154,4 +155,354 @@ namespace UE::ConcertSyncTests::Replication::Stream
 		
 		return true;
 	}
+
+	/**
+	 * Base test for testing ChangeStream requests.
+	 * 
+	 * While there is no sending nor receiving of replicated data, subclasses will still refer to the clients as "Sender"
+	 * and "Receiver".
+	 */
+	class FChangeStreamsTestBase : public FSendReceiveGenericTestBase
+	{
+	public:
+
+		FChangeStreamsTestBase(const FString& InName, const bool bInComplexTask)
+			: FSendReceiveGenericTestBase(InName, bInComplexTask)
+		{}
+
+	protected:
+		
+		ConcertSyncClient::Replication::FJoinReplicatedSessionArgs SenderArgs;
+		ConcertSyncClient::Replication::FJoinReplicatedSessionArgs ReceiverArgs;
+
+		//~ Begin FSendReceiveTestBase Interface
+		virtual ConcertSyncClient::Replication::FJoinReplicatedSessionArgs CreateSenderArgs() override { return SenderArgs; }
+		virtual ConcertSyncClient::Replication::FJoinReplicatedSessionArgs CreateReceiverArgs() override { return ReceiverArgs; }
+		//~ End FSendReceiveTestBase Interfac
+
+		static TTuple<FGuid, FReplicationStreamDescription> CreateFloatPropertyStream(const UTestReflectionObject& TestObject)
+		{
+			const FGuid NewStreamId = FGuid::NewGuid();
+			FConcertPropertySelection Properties;
+			AddFloatProperty(Properties);
+			
+			const FReplicatedObjectInfo AllProperties { TestObject.GetClass(), MoveTemp(Properties) };
+			FReplicationStreamDescription SendingStream;
+			SendingStream.BaseDescription.Identifier = NewStreamId;
+			SendingStream.BaseDescription.ReplicationMap.ReplicatedObjects.Add(&TestObject, AllProperties);
+			return { NewStreamId, SendingStream };
+		}
+
+		static TTuple<FGuid, FReplicationStreamDescription> CreateVectorPropertyStream(const UTestReflectionObject& TestObject)
+		{
+			const FGuid NewStreamId = FGuid::NewGuid();
+			FConcertPropertySelection Properties;
+			AddVectorProperty(Properties);
+			
+			const FReplicatedObjectInfo AllProperties { TestObject.GetClass(), MoveTemp(Properties) };
+			FReplicationStreamDescription SendingStream;
+			SendingStream.BaseDescription.Identifier = NewStreamId;
+			SendingStream.BaseDescription.ReplicationMap.ReplicatedObjects.Add(&TestObject, AllProperties);
+			return { NewStreamId, SendingStream };
+		}
+
+		static FConcertPropertySelection& GetPropertySelection(FReplicationStreamDescription& Stream, const UTestReflectionObject& TestObject)
+		{
+			return Stream.BaseDescription.ReplicationMap.ReplicatedObjects.FindOrAdd(&TestObject).PropertySelection;
+		}
+
+		static void AddFloatProperty(FConcertPropertySelection& PropertySelection)
+		{
+			const FProperty* FloatProperty = UTestReflectionObject::StaticClass()->FindPropertyByName(GET_MEMBER_NAME_CHECKED(UTestReflectionObject, Float));
+			const FConcertPropertyChain PropertyChain(nullptr, *FloatProperty);
+			PropertySelection.ReplicatedProperties.Add(PropertyChain);
+		}
+
+		static void AddVectorProperty(FConcertPropertySelection& PropertySelection)
+		{
+			PropertySelection.ReplicatedProperties.Add(*FConcertPropertyChain::CreateFromPath(*UTestReflectionObject::StaticClass(), { TEXT("Vector") }));
+			PropertySelection.ReplicatedProperties.Add(*FConcertPropertyChain::CreateFromPath(*UTestReflectionObject::StaticClass(), { TEXT("Vector"), TEXT("X") }));
+		}
+
+		/** Sends a ChangeStream request and validates that ExpectedStreams is equal to the registered streams. */
+		void ChangeStreamForSenderClientAndValidate(const FString& InTestName, const ConcertSyncClient::Replication::FChangeStreamRequest& Request, const TArray<FSharedReplicationStreamDescription>& ExpectedStreams)
+		{
+			using namespace ConcertSyncClient::Replication;
+			
+			bool bModifiedInitialStream = false;
+			ClientReplicationManager_Sender->ChangeStream(Request)
+				.Next([this, &InTestName, &bModifiedInitialStream](FChangeStreamResponse&& Response)
+				{
+					bModifiedInitialStream = true;
+					TestTrue(FString::Printf(TEXT("%s > Modify Stream > Success"), *InTestName), Response.IsSuccess());
+				});
+			TestTrue(FString::Printf(TEXT("%s > Modify Stream > Received response"), *InTestName), bModifiedInitialStream);
+			ValidateSenderClientStreams(InTestName, ExpectedStreams);
+		}
+		
+		void ValidateSenderClientStreams(const FString& InTestName, const TArray<FSharedReplicationStreamDescription>& Streams)
+		{
+			using namespace ConcertSyncClient::Replication;
+
+			// Test that remote clients see the receive the same as Streams
+			bool bQueriedClientInfo = false;
+			const FGuid SenderId = Client_Sender->ClientSessionMock->GetSessionClientEndpointId();
+			ClientReplicationManager_Receiver->QueryClientInfo({{ { SenderId } }})
+				.Next([this, &InTestName, &bQueriedClientInfo, &Streams, &SenderId](FClientQueryResponse&& Response)
+				{
+					bQueriedClientInfo = true;
+					const FReplicationClientQueriedInfo* Info = Response.ClientInfo.Find(SenderId);
+					const bool bRegisteredAndQueriedStreamsAreEqual = Info->Streams == Streams;
+					TestTrue(FString::Printf(TEXT("%s > Registered and queried streams are equal"), *InTestName), bRegisteredAndQueriedStreamsAreEqual);
+				});
+
+			// Validate the local client's cache equals  Streams
+			const TArray<FReplicationStreamDescription> LocalStreams = ClientReplicationManager_Sender->GetRegisteredStreams();
+			TArray<FSharedReplicationStreamDescription> TransformedLocalStreams;
+			Algo::Transform(LocalStreams, TransformedLocalStreams, [](const FReplicationStreamDescription& Stream){ return Stream.BaseDescription; });
+			TestEqual(FString::Printf(TEXT("%s > Local streams match registered streams"), *InTestName), TransformedLocalStreams, Streams);
+		}
+	};
+	
+	/** Simple case of FConcertChangeStream_Request where the requesting client has authority and thus cannot generate any authority conflicts. */
+	IMPLEMENT_CUSTOM_SIMPLE_AUTOMATION_TEST(FSimpleChangeStream, FChangeStreamsTestBase, "Concert.Replication.Stream.SimpleChangeStream", EAutomationTestFlags::EditorContext | EAutomationTestFlags::EngineFilter);
+	bool FSimpleChangeStream::RunTest(const FString& Parameters)
+	{
+		using namespace ConcertSyncClient::Replication;
+
+		// 1. Start up client with float property stream
+		UTestReflectionObject* TestObject = NewObject<UTestReflectionObject>(GetTransientPackage());
+		auto[InitialStreamId, InitialStream] = CreateFloatPropertyStream(*TestObject);
+		SenderArgs.Streams = { InitialStream };
+		SetUpClientAndServer();
+
+		// 2.1 Modify existing stream
+		FConcertPropertySelection& ChangedInitialProperties = GetPropertySelection(InitialStream, *TestObject);
+		AddVectorProperty(ChangedInitialProperties);
+		FChangeStreamRequest ModifyRequest;
+		ModifyRequest.ObjectsToPut.Add({ InitialStreamId, TestObject }, { ChangedInitialProperties });
+		ChangeStreamForSenderClientAndValidate(TEXT("ModifyRequest"), ModifyRequest, { InitialStream.BaseDescription });
+
+		// 2.2 Add new stream
+		auto[DynamicStreamId, DynamicStream] = CreateFloatPropertyStream(*TestObject);
+		FChangeStreamRequest DynamicCreationRequest;
+		DynamicCreationRequest.StreamsToAdd.Add(DynamicStream.Pack());
+		ChangeStreamForSenderClientAndValidate(TEXT("DynamicCreationRequest"), DynamicCreationRequest, { InitialStream.BaseDescription, DynamicStream.BaseDescription });
+
+		// 2.3 Remove stream
+		FChangeStreamRequest RemoveStreamRequest;
+		RemoveStreamRequest.StreamsToRemove.Add(InitialStreamId);
+		ChangeStreamForSenderClientAndValidate(TEXT("RemoveStreamRequest"), RemoveStreamRequest, { DynamicStream.BaseDescription });
+
+		// 2.4 Remove object > removes stream implicitly
+		FChangeStreamRequest RemoveObjectRequest;
+		RemoveObjectRequest.ObjectsToRemove.Add({ DynamicStreamId, TestObject });
+		ChangeStreamForSenderClientAndValidate(TEXT("RemoveObjectRequest"), RemoveObjectRequest, {});
+		
+		return true;
+	}
+
+	/** Case of FConcertChangeStream_Request in which requester has authority over an object and creates two streams which have authority over overlapping properties. This is allowed. */
+	IMPLEMENT_CUSTOM_SIMPLE_AUTOMATION_TEST(FChangeStreamWithOverlappingProperties, FChangeStreamsTestBase, "Concert.Replication.Stream.ChangeStreamWithOverlappingProperties", EAutomationTestFlags::EditorContext | EAutomationTestFlags::EngineFilter);
+	bool FChangeStreamWithOverlappingProperties::RunTest(const FString& Parameters)
+	{
+		using namespace ConcertSyncClient::Replication;
+
+		// 1. Start client without any streams
+		UTestReflectionObject* TestObject = NewObject<UTestReflectionObject>(GetTransientPackage());
+		SetUpClientAndServer();
+		
+		// 2.1 Add stream with float property
+		auto[FloatStreamID, FloatStream] = CreateFloatPropertyStream(*TestObject);
+		FChangeStreamRequest CreateFloatStreamRequest;
+		CreateFloatStreamRequest.StreamsToAdd.Add(FloatStream.Pack());
+		ChangeStreamForSenderClientAndValidate(TEXT("CreateFloatStreamRequest"), CreateFloatStreamRequest, { FloatStream.BaseDescription });
+		// 2.2 Add a stream with vector property
+		auto[VectorFloatStreamID, VectorFloatStream] = CreateVectorPropertyStream(*TestObject);
+		FChangeStreamRequest CreateVectorFloatStreamRequest; 
+		CreateVectorFloatStreamRequest.StreamsToAdd.Add(VectorFloatStream.Pack());
+		ChangeStreamForSenderClientAndValidate(TEXT("CreateVectorFloatStreamRequest"), CreateVectorFloatStreamRequest, { FloatStream.BaseDescription, VectorFloatStream.BaseDescription });
+
+		// 2.3 Take authority over both streams
+		bool bTookAuthority = false;
+		FAuthorityChangeRequest TakeAuthorityRequest;
+		TakeAuthorityRequest.TakeAuthority.Add(TestObject, FConcertStreamArray{{ FloatStreamID, VectorFloatStreamID }});
+		ClientReplicationManager_Sender->RequestAuthorityChange({ TakeAuthorityRequest })
+			.Next([this, &bTookAuthority](FAuthorityChangeResponse&& Response)
+			{
+				bTookAuthority = true;
+				TestTrue(TEXT("Authority request > Success"), Response.RejectedObjects.IsEmpty());
+			});
+		TestTrue(TEXT("Authority request > Received response"), bTookAuthority);
+
+		// 2.4 Appending the overlapping float property works
+		FChangeStreamRequest AppendFloatRequest;
+		FConcertPropertySelection& NewSelection = GetPropertySelection(VectorFloatStream, *TestObject);
+		AddFloatProperty(NewSelection);
+		AppendFloatRequest.ObjectsToPut.Add(FObjectInStreamID{ VectorFloatStreamID, TestObject }, FConcertChangeStream_PutObject{ NewSelection });
+		ChangeStreamForSenderClientAndValidate(TEXT("AppendFloatRequest"), AppendFloatRequest, { FloatStream.BaseDescription, VectorFloatStream.BaseDescription });
+		
+		return true;
+	}
+
+	/**
+	 * The requester and client have authority over separate properties on the same object.
+	 * The request would cause the requester's existing authority to overlap with the other client, which is rejected.
+	 */
+	IMPLEMENT_CUSTOM_SIMPLE_AUTOMATION_TEST(FRejectChangeStreamWithConflictingAuthority, FChangeStreamsTestBase, "Concert.Replication.Stream.RejectChangeStreamWithConflictingAuthority", EAutomationTestFlags::EditorContext | EAutomationTestFlags::EngineFilter);
+	bool FRejectChangeStreamWithConflictingAuthority::RunTest(const FString& Parameters)
+	{
+		using namespace ConcertSyncClient::Replication;
+
+		// 1. Start client without any streams
+		UTestReflectionObject* TestObject = NewObject<UTestReflectionObject>(GetTransientPackage());
+		SetUpClientAndServer();
+		
+		// 2.1 Add stream with float property
+		auto[SenderStreamID, SenderStream] = CreateFloatPropertyStream(*TestObject);
+		FChangeStreamRequest CreateFloatStreamRequest;
+		CreateFloatStreamRequest.StreamsToAdd.Add(SenderStream.Pack());
+		ChangeStreamForSenderClientAndValidate(TEXT("CreateFloatStreamRequest"), CreateFloatStreamRequest, { SenderStream.BaseDescription });
+
+		// 2.2 Have Sender take authority over the float property
+		bool bSenderTookAuthority = false;
+		ClientReplicationManager_Sender->TakeAuthorityOver({ TestObject })
+			.Next([this, &bSenderTookAuthority](const FAuthorityChangeResponse& Response) mutable
+			{
+				bSenderTookAuthority = true;
+				TestEqual(TEXT("Sender > No rejection taking authority"), Response.RejectedObjects.Num(), 0); 
+			});
+		TestTrue(TEXT("Sender > Received response taking authority"), bSenderTookAuthority);
+
+		// 2.3 Have Receiver create vector stream
+		auto[ReceiverStreamID, ReceiverStream] = CreateVectorPropertyStream(*TestObject);
+		FChangeStreamRequest CreateVectorStreamRequest; 
+		CreateVectorStreamRequest.StreamsToAdd.Add(ReceiverStream.Pack());
+		bool bAddedVectorStream = false;
+		ClientReplicationManager_Receiver->ChangeStream(CreateVectorStreamRequest)
+			.Next([this, &bAddedVectorStream](FChangeStreamResponse&& Response)
+			{
+				bAddedVectorStream = true;
+				TestTrue(TEXT("Vector Stream > Success"), Response.IsSuccess());
+			});
+		TestTrue(TEXT("Vector Stream> Received response"), bAddedVectorStream);
+		// For simplicity, we don't do any further validation that other clients can see the change
+		
+		// 2.4 Have Receiver take authority over vector stream
+		bool bReceiverTookAuthority = false;
+		ClientReplicationManager_Receiver->TakeAuthorityOver({ TestObject })
+			.Next([this, &bReceiverTookAuthority](const ConcertSyncClient::Replication::FAuthorityChangeResponse& Response) mutable
+			{
+				bReceiverTookAuthority = true;
+				TestEqual(TEXT("Receiver > No rejection taking authority"), Response.RejectedObjects.Num(), 0); 
+			});
+		TestTrue(TEXT("Receiver > Received response taking authority"), bReceiverTookAuthority);
+
+		// 2.5 Receiver is rejected from modifying the stream to include the Float property because the Sender has authority over it
+		const FObjectInStreamID TestObjectInReceiverStreamId { ReceiverStreamID, TestObject };
+		FChangeStreamRequest AddFloatToStreamRequest;
+		FConcertPropertySelection NewSelection = GetPropertySelection(ReceiverStream, *TestObject);
+		AddFloatProperty(NewSelection);
+		AddFloatToStreamRequest.ObjectsToPut.Add(TestObjectInReceiverStreamId, { NewSelection });
+		bool bReceivedResponseAddingFloat = false;
+		
+		// Server logs a warning when rejecting - avoid the test being marked with a warning.
+		AddExpectedError(TEXT("Rejecting ChangeStream"));
+		ClientReplicationManager_Receiver->ChangeStream(AddFloatToStreamRequest)
+			.Next([this, TestObject, SenderStreamID = SenderStreamID, &TestObjectInReceiverStreamId, &bReceivedResponseAddingFloat](FChangeStreamResponse&& Response)
+			{
+				bReceivedResponseAddingFloat = true;
+				TestEqual(TEXT("Append Float > 1 conflict"), Response.AuthorityConflicts.Num(), 1);
+				if (const FReplicatedObjectId* ConflictingObject = Response.AuthorityConflicts.Find(TestObjectInReceiverStreamId))
+				{
+					TestEqual(TEXT("Append float > Conflict > Sender Stream correct"), ConflictingObject->StreamId, SenderStreamID);
+					TestEqual(TEXT("Append float > Conflict > Object correct"), ConflictingObject->Object, FSoftObjectPath(TestObject));
+					TestEqual(TEXT("Append float > Conflict > Endpoint ID correct"), ConflictingObject->SenderEndpointId, Client_Sender->ClientSessionMock->GetSessionClientEndpointId());
+				}
+				else
+				{
+					AddError(TEXT("Expected to find an authority object conflict for TestObject"));
+				}
+				TestTrue(TEXT("Append Float > Failure"), Response.IsFailure());
+			});
+		TestTrue(TEXT("Append Float > Received response"), bReceivedResponseAddingFloat);
+		
+		return true;  
+	}
+
+	/** Tests that no changes are made if a sub-step of FConcertChangeStream_Request causes an error. */
+	IMPLEMENT_CUSTOM_SIMPLE_AUTOMATION_TEST(FChangeStreamIsAtmoic, FChangeStreamsTestBase, "Concert.Replication.Stream.ChangeStreamIsAtmoic", EAutomationTestFlags::EditorContext | EAutomationTestFlags::EngineFilter);
+	bool FChangeStreamIsAtmoic::RunTest(const FString& Parameters)
+	{
+		using namespace ConcertSyncClient::Replication;
+
+		// 1. Start client without any streams
+		UTestReflectionObject* TestObject = NewObject<UTestReflectionObject>(GetTransientPackage());
+		SetUpClientAndServer();
+		
+		// 2.1 Add stream with float property
+		auto[StreamId, Stream] = CreateFloatPropertyStream(*TestObject);
+		FChangeStreamRequest CreateStreamRequest;
+		CreateStreamRequest.StreamsToAdd.Add(Stream.Pack());
+		ChangeStreamForSenderClientAndValidate(TEXT("CreateStreamRequest"), CreateStreamRequest, { Stream.BaseDescription });
+
+		// 2.2 Make a request that will fail and check that no changes were made to the original stream
+		FChangeStreamRequest InvalidRequest;
+		FConcertPropertySelection NewSelection = GetPropertySelection(Stream, *TestObject);
+		AddFloatProperty(NewSelection);
+		InvalidRequest.ObjectsToPut.Add(FObjectInStreamID{ StreamId, TestObject}, FConcertChangeStream_PutObject{ NewSelection });
+		InvalidRequest.StreamsToAdd.Add(Stream.Pack()); // This will make it fail due to pre-existing stream ID
+
+		// Server logs a warning when rejecting - avoid the test being marked with a warning.
+		AddExpectedError(TEXT("Rejecting ChangeStream"));
+		
+		bool bReceivedChangeStreamResponse = false;
+		ClientReplicationManager_Sender->ChangeStream(InvalidRequest)
+			.Next([this, &bReceivedChangeStreamResponse](FChangeStreamResponse&& Response)
+			{
+				bReceivedChangeStreamResponse = true;
+				TestTrue(TEXT("Invalid Request > Modify Stream > Failure"), Response.IsFailure());
+			});
+		
+		TestTrue(TEXT("Invalid Request > Modify Stream > Received response"), bReceivedChangeStreamResponse);
+		ValidateSenderClientStreams(TEXT("Invalid Request"), { Stream.BaseDescription });
+		
+		return true;
+	}
+
+	/** Removing an object from a stream over which the requester has authority, clears the authority the client had. Tests that another client can now claim authority over the previously authored properties. */
+	IMPLEMENT_CUSTOM_SIMPLE_AUTOMATION_TEST(FRemovingObjectClearsAuthority, FChangeStreamsTestBase, "Concert.Replication.Stream.RemovingObjectClearsAuthority", EAutomationTestFlags::EditorContext | EAutomationTestFlags::EngineFilter);
+	bool FRemovingObjectClearsAuthority::RunTest(const FString& Parameters)
+	{
+		using namespace ConcertSyncClient::Replication;
+
+		// 1. Start clients with equivalent streams
+		UTestReflectionObject* TestObject = NewObject<UTestReflectionObject>(GetTransientPackage());
+		auto[SenderStreamId, SenderStream] = CreateFloatPropertyStream(*TestObject);
+		auto[ReceiverStreamId, ReceiverStream] = CreateFloatPropertyStream(*TestObject);
+		SenderArgs.Streams = { SenderStream };
+		ReceiverArgs.Streams = { ReceiverStream };
+		SetUpClientAndServer();
+		
+		// 2.1 Sender takes authority
+		ClientReplicationManager_Sender->TakeAuthorityOver({ TestObject })
+			.Next([this](FAuthorityChangeResponse&& Response)
+			{
+				TestTrue(TEXT("Sender > Success taking authority"), Response.RejectedObjects.IsEmpty());
+			});
+
+		// 2.2 Sender removes the object
+		FChangeStreamRequest RemoveObjectRequest;
+		RemoveObjectRequest.ObjectsToRemove.Add({ SenderStreamId, TestObject });
+		ChangeStreamForSenderClientAndValidate(TEXT("Remove Object"), RemoveObjectRequest, {});
+
+		// 2.3 Receiver can now take authority since 2.2 implicitly removed authority
+		ClientReplicationManager_Receiver->TakeAuthorityOver({ TestObject })
+			.Next([this](FAuthorityChangeResponse&& Response)
+			{
+				TestTrue(TEXT("Receiver > Success taking authority"), Response.RejectedObjects.IsEmpty());
+			});
+		
+		return true;
+	}
 }
