Add DiffChanges function which builds FConcertChangeStream_Request by inspecting the difference between two streams

diff --git a/Engine/Plugins/Developer/Concert/ConcertSync/ConcertSyncCore/Source/ConcertSyncCore/Private/Replication/ChangeStreamSharedUtils.cpp b/Engine/Plugins/Developer/Concert/ConcertSync/ConcertSyncCore/Source/ConcertSyncCore/Private/Replication/ChangeStreamSharedUtils.cpp
--- a/Engine/Plugins/Developer/Concert/ConcertSync/ConcertSyncCore/Source/ConcertSyncCore/Private/Replication/ChangeStreamSharedUtils.cpp
+++ b/Engine/Plugins/Developer/Concert/ConcertSync/ConcertSyncCore/Source/ConcertSyncCore/Private/Replication/ChangeStreamSharedUtils.cpp
@@ -102,4 +102,82 @@ namespace UE::ConcertSyncCore::Replication::ChangeStreamUtils
 			}
 		}
 	}
+	
+	void IterateInvalidEntries(
+		const FObjectReplicationMap& ReplicationMap,
+		TFunctionRef<EBreakBehavior(const FSoftObjectPath&, const FReplicatedObjectInfo&)> Callback
+		)
+	{
+		for (const TPair<const FSoftObjectPath&, const FReplicatedObjectInfo&> Pair : ReplicationMap.ReplicatedObjects)
+		{
+			if (!Pair.Value.IsValidForSendingToServer() && Callback(Pair.Key, Pair.Value) == EBreakBehavior::Break)
+			{
+				break;
+			}
+		}
+	}
+
+	namespace Private
+	{
+		static void BuildPutObjectList(const FGuid& StreamId, const FObjectReplicationMap& Base, const FObjectReplicationMap& Desired, FConcertChangeStream_Request& Request)
+		{
+			for (const TPair<FSoftObjectPath, FReplicatedObjectInfo>& BasePair : Base.ReplicatedObjects)
+			{
+				const FSoftObjectPath& ObjectPath = BasePair.Key;
+				const FObjectInStreamID ObjectId { StreamId, ObjectPath };
+				const FReplicatedObjectInfo* DesiredObjectInfo = Desired.ReplicatedObjects.Find(ObjectPath);
+				if (DesiredObjectInfo)
+				{
+					const FReplicatedObjectInfo& BaseObjectInfo = BasePair.Value;
+					const TOptional<FConcertChangeStream_PutObject> PutObject = FConcertChangeStream_PutObject::MakeFromChange(BaseObjectInfo, *DesiredObjectInfo);
+				
+					const bool bDesiredHasChangedFromBase = BaseObjectInfo != *DesiredObjectInfo;
+					// If MakeFromChange returned unset, it means that this request is not valid to submit. 
+					const bool bBaseAndDesiredStateAreValid = ensureMsgf(PutObject, TEXT("Function assumption violated; you did not pass in valid base or desired state."));
+					if (bDesiredHasChangedFromBase && bBaseAndDesiredStateAreValid)
+					{
+						Request.ObjectsToPut.Add(ObjectId, *PutObject);
+					}
+				}
+				else
+				{
+					Request.ObjectsToRemove.Add(ObjectId);
+				}
+			}
+		}
+
+		static void BuildRemoveObjectList(const FGuid& StreamId, const FObjectReplicationMap& Base, const FObjectReplicationMap& Desired, FConcertChangeStream_Request& Request)
+		{
+			for (const TPair<FSoftObjectPath, FReplicatedObjectInfo>& DesiredPair : Desired.ReplicatedObjects)
+			{
+				const FSoftObjectPath& ObjectPath = DesiredPair.Key;
+				if (Base.ReplicatedObjects.Contains(ObjectPath))
+				{
+					// Handled up above
+					continue;
+				}
+
+				// Desired wants to add an object
+				const FReplicatedObjectInfo& ObjectInfo = DesiredPair.Value;
+				const TOptional<FConcertChangeStream_PutObject> PutObject = FConcertChangeStream_PutObject::MakeFromInfo(ObjectInfo);
+				const bool bDesiredStateHasEnoughDataToForPut = ensureMsgf(PutObject, TEXT("Function assumption violated; you did not pass in valid base or desired state."));
+				if (bDesiredStateHasEnoughDataToForPut)
+				{
+					Request.ObjectsToPut.Add({ StreamId, ObjectPath}, *PutObject);
+				}
+			}
+		}
+	}
+
+	FConcertChangeStream_Request BuildRequestFromDiff(
+		const FGuid& StreamId,
+		const FObjectReplicationMap& Base,
+		const FObjectReplicationMap& Desired
+		)
+	{
+		FConcertChangeStream_Request Request;
+		Private::BuildPutObjectList(StreamId, Base, Desired, Request);
+		Private::BuildRemoveObjectList(StreamId, Base, Desired, Request);
+		return Request;
+	}
 }
diff --git a/Engine/Plugins/Developer/Concert/ConcertSync/ConcertSyncCore/Source/ConcertSyncCore/Private/Replication/Messages/ConcertReplicationEvents.cpp b/Engine/Plugins/Developer/Concert/ConcertSync/ConcertSyncCore/Source/ConcertSyncCore/Private/Replication/Messages/ConcertReplicationEvents.cpp
new file mode 100644
--- /dev/null
+++ b/Engine/Plugins/Developer/Concert/ConcertSync/ConcertSyncCore/Source/ConcertSyncCore/Private/Replication/Messages/ConcertReplicationEvents.cpp
@@ -0,0 +1,61 @@
+// Copyright Epic Games, Inc. All Rights Reserved.
+
+#include "Replication/Messages/ConcertReplicationEvents.h"
+
+#include "Misc/OutputDevice.h"
+
+TOptional<FConcertChangeStream_PutObject> FConcertChangeStream_PutObject::MakeFromInfo(const FReplicatedObjectInfo& New)
+{
+	if (!New.IsValidForSendingToServer())
+	{
+		return {};
+	}
+	return FConcertChangeStream_PutObject{ New.PropertySelection, New.ClassPath };
+}
+
+TOptional<FConcertChangeStream_PutObject> FConcertChangeStream_PutObject::MakeFromChange(const FReplicatedObjectInfo& Base, const FReplicatedObjectInfo& Desired)
+{
+	if (Base.IsValidForSendingToServer() && Desired.IsValidForSendingToServer())
+	{
+		const FConcertPropertySelection PropertySelection = Base.PropertySelection != Desired.PropertySelection ? Desired.PropertySelection : FConcertPropertySelection{};
+		const FSoftClassPath ClassPath = Base.ClassPath != Desired.ClassPath ? Desired.ClassPath : FSoftClassPath{};
+		return FConcertChangeStream_PutObject{ PropertySelection, ClassPath };
+	}
+	return {};
+}
+
+TOptional<FReplicatedObjectInfo> FConcertChangeStream_PutObject::MakeObjectInfoIfValid() const
+{
+	if (Properties.ReplicatedProperties.IsEmpty() || ClassPath.IsNull())
+	{
+		return {};
+	}
+	return FReplicatedObjectInfo{ ClassPath, Properties };
+}
+
+void FConcertChangeStream_Response::LogErrors(FOutputDevice& OutputDevice) const
+{
+	for (const TPair<FObjectInStreamID, FReplicatedObjectId>& Conflict : AuthorityConflicts)
+	{
+		OutputDevice.Logf(TEXT("Authority: { %s } conflicts with { %s }"), *Conflict.Key.ToString(), *Conflict.Value.ToString());
+	}
+
+	for (const TPair<FObjectInStreamID, EConcertPutObjectErrorCode>& Error : ObjectsToPutSemanticErrors)
+	{
+		const FString ErrorCodeAsString = [&Error]()
+		{
+			switch (Error.Value)
+			{
+			case EConcertPutObjectErrorCode::UnresolvedStream: return TEXT("Unresolved stream");
+			case EConcertPutObjectErrorCode::MissingData: return TEXT("Missing data");
+			default: return TEXT("Unknown");
+			}
+		}();
+		OutputDevice.Logf(TEXT("Semantic error: %s - %s"), *Error.Key.ToString(), *ErrorCodeAsString);
+	}
+
+	for (const FGuid& Stream : FailedStreamCreation)
+	{
+		OutputDevice.Logf(TEXT("Failed to create stream: %s"), *Stream.ToString());
+	}
+}
diff --git a/Engine/Plugins/Developer/Concert/ConcertSync/ConcertSyncCore/Source/ConcertSyncCore/Public/Replication/ChangeStreamSharedUtils.h b/Engine/Plugins/Developer/Concert/ConcertSync/ConcertSyncCore/Source/ConcertSyncCore/Public/Replication/ChangeStreamSharedUtils.h
--- a/Engine/Plugins/Developer/Concert/ConcertSync/ConcertSyncCore/Source/ConcertSyncCore/Public/Replication/ChangeStreamSharedUtils.h
+++ b/Engine/Plugins/Developer/Concert/ConcertSync/ConcertSyncCore/Source/ConcertSyncCore/Public/Replication/ChangeStreamSharedUtils.h
@@ -10,7 +10,9 @@ enum class EBreakBehavior : uint8;
 
 struct FConcertChangeStream_Request;
 struct FObjectInStreamID;
+struct FObjectReplicationMap;
 struct FReplicationStreamDescription;
+struct FReplicatedObjectInfo;
 
 /** This namespace contains the shared logic for applying FConcertChangeStream_Requests. */
 namespace UE::ConcertSyncCore::Replication::ChangeStreamUtils
@@ -33,5 +35,30 @@ namespace UE::ConcertSyncCore::Replication::ChangeStreamUtils
 	 * @param Request The request to parse
 	 * @param StreamsToModify The streams to apply the request to
 	 */
-	CONCERTSYNCCORE_API void ApplyValidatedRequest(const FConcertChangeStream_Request& Request, IN OUT TArray<FReplicationStreamDescription>& StreamsToModify);
+	CONCERTSYNCCORE_API void ApplyValidatedRequest(
+		const FConcertChangeStream_Request& Request,
+		IN OUT TArray<FReplicationStreamDescription>& StreamsToModify
+		);
+
+	/** Iterates entries with incomplete data, which are entries for which IsValidForSendingToServer returns false. */
+	CONCERTSYNCCORE_API void IterateInvalidEntries(const FObjectReplicationMap& ReplicationMap, TFunctionRef<EBreakBehavior(const FSoftObjectPath&, const FReplicatedObjectInfo&)> Callback);
+
+	/**
+	 * Builds PutObject and RemoveObject requests that would transform Base to Desired.
+	 * 
+	 * The state of an object in a replication map is only valid if all its members have values:
+	 * that means the class is non-null and not the property selection is non-empty.
+	 * This function enforces (and ensure()s) that the Base and Desired states are valid.
+	 * The changes are validated using FConcertChangeStream_PutObject's factories MakeFromInfo and MakeFromChange.
+	 * 
+	 * @param StreamId The ID of the changed stream: it is inserted into the request
+	 * @param Base The base stream
+	 * @param Desired The state the stream should end up in 
+	 * @return A semantically correct FConcertChangeStream_Request which Base to Desired (note is can be rejected due to authority conflicts).
+	 */
+	CONCERTSYNCCORE_API FConcertChangeStream_Request BuildRequestFromDiff(
+		const FGuid& StreamId,
+		const FObjectReplicationMap& Base,
+		const FObjectReplicationMap& Desired
+		);
 }
diff --git a/Engine/Plugins/Developer/Concert/ConcertSync/ConcertSyncCore/Source/ConcertSyncCore/Public/Replication/Data/ObjectIds.h b/Engine/Plugins/Developer/Concert/ConcertSync/ConcertSyncCore/Source/ConcertSyncCore/Public/Replication/Data/ObjectIds.h
--- a/Engine/Plugins/Developer/Concert/ConcertSync/ConcertSyncCore/Source/ConcertSyncCore/Public/Replication/Data/ObjectIds.h
+++ b/Engine/Plugins/Developer/Concert/ConcertSync/ConcertSyncCore/Source/ConcertSyncCore/Public/Replication/Data/ObjectIds.h
@@ -18,6 +18,8 @@ struct FObjectInStreamID
 	/** The object to process */
 	UPROPERTY()
 	FSoftObjectPath Object;
+
+	FString ToString() const { return FString::Printf(TEXT("StreamId: %s, Object: %s"), *StreamId.ToString(), *Object.ToString()); }
 		
 	friend bool operator==(const FObjectInStreamID& Left, const FObjectInStreamID& Right)
 	{
@@ -38,6 +40,8 @@ struct FReplicatedObjectId : public FObjectInStreamID
 	
 	/** The ID of the endpoint that sent this object. This can be the client or server endpoint depending on who receives it. */
 	FGuid SenderEndpointId;
+	
+	FString ToString() const { return FString::Printf(TEXT("StreamId: %s, Object: %s, Sender: %s"), *StreamId.ToString(), *Object.ToString(), *SenderEndpointId.ToString()); }
 		
 	friend bool operator==(const FReplicatedObjectId& Left, const FReplicatedObjectId& Right)
 	{
diff --git a/Engine/Plugins/Developer/Concert/ConcertSync/ConcertSyncCore/Source/ConcertSyncCore/Public/Replication/Data/ObjectReplicationMap.h b/Engine/Plugins/Developer/Concert/ConcertSync/ConcertSyncCore/Source/ConcertSyncCore/Public/Replication/Data/ObjectReplicationMap.h
--- a/Engine/Plugins/Developer/Concert/ConcertSync/ConcertSyncCore/Source/ConcertSyncCore/Public/Replication/Data/ObjectReplicationMap.h
+++ b/Engine/Plugins/Developer/Concert/ConcertSync/ConcertSyncCore/Source/ConcertSyncCore/Public/Replication/Data/ObjectReplicationMap.h
@@ -16,12 +16,14 @@ struct FReplicatedObjectInfo
 	UPROPERTY()
 	FConcertPropertySelection PropertySelection;
 
+	/** @return Whether this data is valid for sending to the server. */
+	bool IsValidForSendingToServer() const { return ClassPath.IsValid() && !PropertySelection.ReplicatedProperties.IsEmpty(); }
+	
 	friend bool operator==(const FReplicatedObjectInfo& Left, const FReplicatedObjectInfo& Right)
 	{
 		return Left.ClassPath == Right.ClassPath
 			&& Left.PropertySelection == Right.PropertySelection;
 	}
-
 	friend bool operator!=(const FReplicatedObjectInfo& Left, const FReplicatedObjectInfo& Right)
 	{
 		return !(Left == Right);
diff --git a/Engine/Plugins/Developer/Concert/ConcertSync/ConcertSyncCore/Source/ConcertSyncCore/Public/Replication/Messages/ConcertReplicationEvents.h b/Engine/Plugins/Developer/Concert/ConcertSync/ConcertSyncCore/Source/ConcertSyncCore/Public/Replication/Messages/ConcertReplicationEvents.h
--- a/Engine/Plugins/Developer/Concert/ConcertSync/ConcertSyncCore/Source/ConcertSyncCore/Public/Replication/Messages/ConcertReplicationEvents.h
+++ b/Engine/Plugins/Developer/Concert/ConcertSync/ConcertSyncCore/Source/ConcertSyncCore/Public/Replication/Messages/ConcertReplicationEvents.h
@@ -5,8 +5,11 @@
 #include "ConcertMessageData.h"
 #include "Replication/Data/ClientQueriedInfo.h"
 #include "Replication/Data/ObjectIds.h"
+#include "Misc/Optional.h"
 #include "ConcertReplicationEvents.generated.h"
 
+class FOutputDevice;
+
 /** Contains data to be applied to a replicated object. */
 USTRUCT()
 struct FConcertObjectReplicationEvent
@@ -204,10 +207,21 @@ struct FConcertChangeStream_PutObject
 	 */
 	UPROPERTY()
 	FSoftClassPath ClassPath;
+
+	// Intention: Ideally code dealing with PutObject requests uses these constructors / factory functions.
+	// If a property is added to FReplicatedObjectInfo, only the below code needs to be updated.
+
+	/** @return The PutObject request if New contained sufficient info - empty otherwise */
+	CONCERTSYNCCORE_API static TOptional<FConcertChangeStream_PutObject> MakeFromInfo(const FReplicatedObjectInfo& New);
+	/** @return The PutObject request if changing Base to Desired contained sufficient info - empty otherwise */
+	CONCERTSYNCCORE_API static TOptional<FConcertChangeStream_PutObject> MakeFromChange(const FReplicatedObjectInfo& Base, const FReplicatedObjectInfo& Desired);
+
+	/** Creates a new object info if there is sufficient data (all fields must be set for this). */
+	CONCERTSYNCCORE_API TOptional<FReplicatedObjectInfo> MakeObjectInfoIfValid() const;
 };
 
 /**
- * Let's a client changing its streams owned on the server.
+ * Let's a client change its streams owned on the server.
  * 
  * This request is processed atomically: it either succeeds completely or fails completely.
  * If any sub-change causes a failure, the failure reason will be returned in the response and the client must make a new request.
@@ -235,6 +249,9 @@ struct FConcertChangeStream_Request
 	 * @see FConcertChangeStream_Response::AuthorityConflicts for some examples of conflicts.
 	 *
 	 * If the key identifies a stream that does not exist, the request will fail.
+	 * 
+	 * If the key identifies a stream that is in StreamsToAdd, the request will fail to avoid allowing the construction of
+	 * ambiguous requests, e.g. both StreamsToAdd and ObjectsToPut containing object Foo but with different property selections.
 	 */
 	UPROPERTY()
 	TMap<FObjectInStreamID, FConcertChangeStream_PutObject> ObjectsToPut;
@@ -242,6 +259,9 @@ struct FConcertChangeStream_Request
 	/**
 	 * New streams to add to the server.
 	 * Fails if any ID overlaps with a pre-existing one.
+	 *
+	 * It is valid to have the same stream in StreamsToRemove: StreamsToAdd is applied after StreamsToRemove resulting
+	 * in the stream's content being replaced.
 	 */
 	UPROPERTY()
 	TArray<FReplicationStreamDescription_NetPacked> StreamsToAdd;
@@ -251,6 +271,9 @@ struct FConcertChangeStream_Request
 	 * Supplying a stream that does not exist does not cause failure (but is nonsensical).
 	 * 
 	 * If the requesting client has authority over any of the objects contained in the stream, authority is removed.
+	 * 
+	 * It is valid to have the same stream in StreamsToAdd: StreamsToAdd is applied after StreamsToRemove resulting
+	 * in the stream's content being replaced.
 	 */
 	UPROPERTY()
 	TSet<FGuid> StreamsToRemove;
@@ -305,4 +328,7 @@ struct FConcertChangeStream_Response
 
 	bool IsSuccess() const { return AuthorityConflicts.IsEmpty() && ObjectsToPutSemanticErrors.IsEmpty() && FailedStreamCreation.IsEmpty(); }
 	bool IsFailure() const { return !IsSuccess(); }
+
+	/** If IsFailure(), logs the errors. */
+	CONCERTSYNCCORE_API void LogErrors(FOutputDevice& OutputDevice) const;
 };
\ No newline at end of file
diff --git a/Engine/Plugins/Developer/Concert/ConcertSync/ConcertSyncTest/Source/ConcertSyncTest/Private/Replication/Components/ReplicationCacheTests.cpp b/Engine/Plugins/Developer/Concert/ConcertSync/ConcertSyncTest/Source/ConcertSyncTest/Private/Replication/Components/ReplicationCacheTests.cpp
new file mode 100644
--- /dev/null
+++ b/Engine/Plugins/Developer/Concert/ConcertSync/ConcertSyncTest/Source/ConcertSyncTest/Private/Replication/Components/ReplicationCacheTests.cpp
@@ -0,0 +1,172 @@
+// Copyright Epic Games, Inc. All Rights Reserved.
+
+#include "Replication/Formats/IObjectReplicationFormat.h"
+#include "Replication/Messages/ConcertReplicationEvents.h"
+#include "Replication/Processing/ObjectReplicationCache.h"
+#include "Replication/Data/ObjectIds.h"
+#include "Replication/TestReflectionObject.h"
+#include "Util/ClientServerCommunicationTest.h"
+
+#include "Misc/AutomationTest.h"
+#include "Misc/Paths.h"
+#include "HAL/FileManager.h"
+
+namespace UE::ConcertSyncTests
+{
+	enum class EReplicationCacheTestFlags
+	{
+		None = 0,
+		InstantConsume = 1 << 0,
+		NeverConsume = 1 << 1,
+		NeverReceive = 1 << 2,
+	};
+	ENUM_CLASS_FLAGS(EReplicationCacheTestFlags);
+
+	/**
+	 * Tests that FObjectReplicationCache works as expected.
+	 */
+	IMPLEMENT_SIMPLE_AUTOMATION_TEST(FReplicationCacheTest, "Concert.Replication.Components.ReplicationCache", EAutomationTestFlags::EditorContext | EAutomationTestFlags::EngineFilter);
+	bool FReplicationCacheTest::RunTest(const FString& Parameters)
+	{
+		class FTestReplicationFormat : public ConcertSyncCore::IObjectReplicationFormat
+		{
+		public:
+			virtual TOptional<FConcertSessionSerializedPayload> CreateReplicationEvent(UObject& Object, FAllowPropertyFunc IsPropertyAllowedFunc) override { return NotMocked<TOptional<FConcertSessionSerializedPayload>>({}); }
+			virtual void ClearInternalCache(TArrayView<UObject> ObjectsToClear) override { return NotMocked<void>(); }
+			virtual void CombineReplicationEvents(FConcertSessionSerializedPayload& Base, const FConcertSessionSerializedPayload& Newer) override
+			{
+				// Reuse FNativeStruct to avoid introducing more test types
+				FNativeStruct BaseStruct;
+				Base.GetTypedPayload(BaseStruct);
+				FNativeStruct NewerStruct;
+				Newer.GetTypedPayload(NewerStruct);
+
+				BaseStruct.Float += NewerStruct.Float;
+				Base.SetTypedPayload(BaseStruct);
+			}
+			virtual void ApplyReplicationEvent(UObject& Object, const FConcertSessionSerializedPayload& Payload) override { return NotMocked<void>(); }
+
+			static FConcertObjectReplicationEvent CreateEvent(FSoftObjectPath Object, float Value)
+			{
+				FNativeStruct Data{ Value };
+				FConcertSessionSerializedPayload Payload;
+				Payload.SetTypedPayload(Data, EConcertPayloadCompressionType::None);
+				return { MoveTemp(Object), Payload };
+			}
+		};
+
+		class FTestReplicationCacheUser : public ConcertSyncCore::IReplicationCacheUser
+		{
+		public:
+
+			FAutomationTestBase& Test;
+			FObjectInStreamID AllowedObject;
+			EReplicationCacheTestFlags Flags;
+			TSharedPtr<const FConcertObjectReplicationEvent> CachedData;
+
+			FTestReplicationCacheUser(FAutomationTestBase& Test, FObjectInStreamID AllowedObject, EReplicationCacheTestFlags Flags = EReplicationCacheTestFlags::None)
+				: Test(Test)
+				, AllowedObject(MoveTemp(AllowedObject))
+				, Flags(Flags)
+			{}
+			
+			virtual bool WantsToAcceptObject(const FReplicatedObjectId& Object) const override
+			{
+				return !EnumHasAnyFlags(Flags, EReplicationCacheTestFlags::NeverReceive) && AllowedObject == Object;
+			}
+			
+			virtual void OnDataCached(const FReplicatedObjectId& Object, TSharedRef<const FConcertObjectReplicationEvent> Data) override
+			{
+				if (EnumHasAnyFlags(Flags, EReplicationCacheTestFlags::NeverReceive))
+				{
+					Test.AddError(TEXT("Received object that we never asked for it!"));
+					return;
+				}
+				
+				if (EnumHasAnyFlags(Flags, EReplicationCacheTestFlags::NeverConsume) && CachedData.IsValid())
+				{
+					Test.AddError(TEXT("Received new object although the old one was not consumed yet!"));
+					return;
+				}
+				
+				if (!EnumHasAnyFlags(Flags, EReplicationCacheTestFlags::InstantConsume))
+				{
+					CachedData = MoveTemp(Data);
+				}
+			}
+
+			float PeakData()
+			{
+				if (!CachedData)
+				{
+					return -1.f;
+				}
+
+				FNativeStruct Data;
+				CachedData->SerializedPayload.GetTypedPayload(Data);
+				return Data.Float;
+			}
+
+			const void* GetDataAddress()
+			{
+				return CachedData.Get();
+			}
+			
+			void Consume()
+			{
+				check(!EnumHasAnyFlags(Flags, EReplicationCacheTestFlags::NeverConsume));
+				CachedData.Reset();
+			}
+		};
+
+
+		// Set up
+		const FGuid StreamId = FGuid::NewGuid();
+		const FSoftObjectPath ObjectPath(TEXT("/Game/World.World:PersistentLevel.StaticMeshActor0"));
+		const FGuid DummySendingClientId = FGuid::NewGuid();
+		const FReplicatedObjectId ObjectID{ { StreamId, ObjectPath }, DummySendingClientId};
+		
+		TSharedRef<ConcertSyncCore::FObjectReplicationCache> Cache = MakeShared<ConcertSyncCore::FObjectReplicationCache>(MakeShared<FTestReplicationFormat>());
+		TSharedRef<FTestReplicationCacheUser> User_NeverConsume = MakeShared<FTestReplicationCacheUser>(*this, ObjectID, EReplicationCacheTestFlags::NeverConsume);
+		TSharedRef<FTestReplicationCacheUser> User_NeverReceive = MakeShared<FTestReplicationCacheUser>(*this, ObjectID, EReplicationCacheTestFlags::NeverReceive);
+		TSharedRef<FTestReplicationCacheUser> User_ConsumeManually = MakeShared<FTestReplicationCacheUser>(*this, ObjectID);
+		TSharedRef<FTestReplicationCacheUser> User_InstantConsume = MakeShared<FTestReplicationCacheUser>(*this, ObjectID, EReplicationCacheTestFlags::InstantConsume);
+		
+		Cache->RegisterDataCacheUser(User_NeverConsume);
+		Cache->RegisterDataCacheUser(User_NeverReceive);
+		Cache->RegisterDataCacheUser(User_ConsumeManually);
+		Cache->RegisterDataCacheUser(User_InstantConsume);
+
+		const FConcertObjectReplicationEvent Event_5 = FTestReplicationFormat::CreateEvent(ObjectPath, 5.f);
+		const FConcertObjectReplicationEvent Event_10 = FTestReplicationFormat::CreateEvent(ObjectPath, 10.f);
+		const FConcertObjectReplicationEvent Event_20 = FTestReplicationFormat::CreateEvent(ObjectPath, 20.f);
+		const FConcertObjectReplicationEvent Event_100 = FTestReplicationFormat::CreateEvent(ObjectPath, 100.f);
+
+
+		
+		// Tests
+		Cache->StoreUntilConsumed(DummySendingClientId, StreamId, Event_5);
+		TestTrue(TEXT("Users have same data"), User_NeverConsume->CachedData.Get() == User_ConsumeManually->CachedData.Get());
+		User_ConsumeManually->Consume();
+		
+		Cache->StoreUntilConsumed(DummySendingClientId, StreamId, Event_10);
+		TestEqual(TEXT("Combined events: 5 and 10"), User_NeverConsume->PeakData(), 15.f);
+		TestEqual(TEXT("Received new data: 10"), User_ConsumeManually->PeakData(), 10.f);
+		
+		// The goal here is to test that the cache does not leak resources when a user de-registers
+		Cache->UnregisterDataCacheUser(User_ConsumeManually);
+		Cache->RegisterDataCacheUser(User_ConsumeManually);
+		const void* AddressBefore = User_ConsumeManually->GetDataAddress();
+		Cache->StoreUntilConsumed(DummySendingClientId, StreamId, Event_100);
+		TestEqual(TEXT("Combined events: 5, 10, and 100"), User_NeverConsume->PeakData(), 115.f);
+		TestEqual(TEXT("Received totally new data"), User_ConsumeManually->PeakData(), 100.f);
+		TestNotEqual(TEXT("Received totally new allocated data"), AddressBefore, User_ConsumeManually->GetDataAddress());
+
+		// Test: No crash or anything
+		Cache->UnregisterDataCacheUser(User_NeverConsume);
+		Cache->UnregisterDataCacheUser(User_NeverReceive);
+		Cache->UnregisterDataCacheUser(User_ConsumeManually);
+		Cache->UnregisterDataCacheUser(User_InstantConsume);
+		return true;
+	}
+}
diff --git a/Engine/Plugins/Developer/Concert/ConcertSync/ConcertSyncTest/Source/ConcertSyncTest/Private/Replication/Components/ReplicationChangeStreamUtilsTests.cpp b/Engine/Plugins/Developer/Concert/ConcertSync/ConcertSyncTest/Source/ConcertSyncTest/Private/Replication/Components/ReplicationChangeStreamUtilsTests.cpp
new file mode 100644
--- /dev/null
+++ b/Engine/Plugins/Developer/Concert/ConcertSync/ConcertSyncTest/Source/ConcertSyncTest/Private/Replication/Components/ReplicationChangeStreamUtilsTests.cpp
@@ -0,0 +1,101 @@
+// Copyright Epic Games, Inc. All Rights Reserved.
+
+#include "Replication/ChangeStreamSharedUtils.h"
+#include "Replication/Data/ConcertPropertySelection.h"
+#include "Replication/Data/ObjectReplicationMap.h"
+#include "Replication/Messages/ConcertReplicationEvents.h"
+
+#include "Misc/AutomationTest.h"
+#include "Components/SceneComponent.h"
+#include "Components/StaticMeshComponent.h"
+
+/**
+ * Tests for UE::ConcertSyncCore::Replication::ChangeStreamUtils
+ */
+namespace UE::ConcertSyncTests::ChangeStreamUtils
+{
+	/**
+	 * Tests that UE::ConcertSyncCore::Replication::ChangeStreamUtils::DiffChanges detects
+	 *  - adding and removing objects
+	 *  - modify properties and class
+	 */
+	IMPLEMENT_SIMPLE_AUTOMATION_TEST(FSimpleDiffChangesTest, "Concert.Replication.Components.DiffChangesSimple", EAutomationTestFlags::EditorContext | EAutomationTestFlags::EngineFilter);
+	bool FSimpleDiffChangesTest::RunTest(const FString& Parameters)
+	{
+		// 1. Set up
+		const FGuid StreamId = FGuid::NewGuid();
+		const FSoftClassPath SceneComponentClass = USceneComponent::StaticClass();
+		const FSoftClassPath StaticMeshComponentClass = UStaticMeshComponent::StaticClass();
+		const FSoftObjectPath ChangeProperties(TEXT("/Game/World.World:PersistentLevel.ChangeProperties.StaticMeshComponent"));
+		const FSoftObjectPath ChangeClass(TEXT("/Game/World.World:PersistentLevel.ChangeClass.StaticMeshComponent"));
+		const FSoftObjectPath RemoveMe(TEXT("/Game/World.World:PersistentLevel.RemoveMe.StaticMeshComponent"));
+		const FSoftObjectPath AddMe(TEXT("/Game/World.World:PersistentLevel.AddMe.StaticMeshComponent"));
+
+		const TOptional<FConcertPropertyChain> RelativeLocation = FConcertPropertyChain::CreateFromPath(*USceneComponent::StaticClass(), { TEXT("RelativeLocation") });
+		const TOptional<FConcertPropertyChain> RelativeLocationX = FConcertPropertyChain::CreateFromPath(*USceneComponent::StaticClass(), { TEXT("RelativeLocation"), TEXT("X") });
+		const TOptional<FConcertPropertyChain> RelativeLocationY = FConcertPropertyChain::CreateFromPath(*USceneComponent::StaticClass(), { TEXT("RelativeLocation"), TEXT("Y") });
+		const TOptional<FConcertPropertyChain> RelativeLocationZ = FConcertPropertyChain::CreateFromPath(*USceneComponent::StaticClass(), { TEXT("RelativeLocation"), TEXT("Z") });
+		if (!ensureAlways(RelativeLocation && RelativeLocationX && RelativeLocationY && RelativeLocationZ))
+		{
+			return false;
+		}
+		const FConcertPropertySelection Selection_RelativeLocation { { *RelativeLocation, *RelativeLocationX, *RelativeLocationY, *RelativeLocationZ } };
+		const FConcertPropertySelection Selection_RelativeLocationOnlyX { { *RelativeLocation, *RelativeLocationX} };
+		
+
+		// 2. Run
+		FObjectReplicationMap Base;
+		Base.ReplicatedObjects.Add(ChangeProperties, { SceneComponentClass, Selection_RelativeLocation });
+		Base.ReplicatedObjects.Add(ChangeClass, { SceneComponentClass, Selection_RelativeLocation });
+		Base.ReplicatedObjects.Add(RemoveMe, { SceneComponentClass, Selection_RelativeLocation });
+		FObjectReplicationMap Desired;
+		Desired.ReplicatedObjects.Add(ChangeProperties, { SceneComponentClass, Selection_RelativeLocationOnlyX });
+		Desired.ReplicatedObjects.Add(ChangeClass, { StaticMeshComponentClass, Selection_RelativeLocation });
+		Desired.ReplicatedObjects.Add(AddMe, { SceneComponentClass, Selection_RelativeLocation });
+		const FConcertChangeStream_Request Request = ConcertSyncCore::Replication::ChangeStreamUtils::BuildRequestFromDiff(StreamId, Base, Desired);
+
+		
+		// 3. Test
+		constexpr int32 NumExpectedPutObjects = 3; // ChangeProperties + ChangeClass + AddMe
+		TestEqual(TEXT("3 objects were put"), Request.ObjectsToPut.Num(), NumExpectedPutObjects);
+		TestEqual(TEXT("No added streams"), Request.StreamsToAdd.Num(), 0);
+		TestEqual(TEXT("No removed streams"), Request.StreamsToRemove.Num(), 0);
+		
+		TestEqual(TEXT("Remove exactly 1"), Request.ObjectsToRemove.Num(), 1);
+		TestTrue(TEXT("Remove correct object"), Request.ObjectsToRemove.Contains({ StreamId, RemoveMe }));
+
+		if (const FConcertChangeStream_PutObject* ChangePropertiesChange = Request.ObjectsToPut.Find({ StreamId, ChangeProperties }))
+		{
+			// No class change, so it should be null in the request
+			TestTrue(TEXT("Change Properties > Class unset"), ChangePropertiesChange->ClassPath.IsNull());
+			TestEqual(TEXT("Change Properties > Properties"), ChangePropertiesChange->Properties, Selection_RelativeLocationOnlyX);
+		}
+		else
+		{
+			AddError(TEXT("Did not change ChangeProperties object"));
+		}
+
+		if (const FConcertChangeStream_PutObject* ChangeClassChange = Request.ObjectsToPut.Find({ StreamId, ChangeClass }))
+		{
+			TestEqual(TEXT("Change Class > Class"), ChangeClassChange->ClassPath, StaticMeshComponentClass);
+			// No properties changed, so it should be empty in the request
+			TestTrue(TEXT("Change Class > Properties empty"), ChangeClassChange->Properties.ReplicatedProperties.IsEmpty());
+		}
+		else
+		{
+			AddError(TEXT("Did not change ChangeClass object"));
+		}
+
+		if (const FConcertChangeStream_PutObject* AddMeChange = Request.ObjectsToPut.Find({ StreamId, AddMe }))
+		{
+			TestEqual(TEXT("Add Object > Class"), AddMeChange->ClassPath, SceneComponentClass);
+			TestEqual(TEXT("Add Object > Properties"), AddMeChange->Properties, Selection_RelativeLocation);
+		}
+		else
+		{
+			AddError(TEXT("Did not add AddMe object"));
+		}
+		
+		return true;
+	}
+}
\ No newline at end of file
diff --git a/Engine/Plugins/Developer/Concert/ConcertSync/ConcertSyncTest/Source/ConcertSyncTest/Private/Replication/ReplicationCacheTests.cpp b/Engine/Plugins/Developer/Concert/ConcertSync/ConcertSyncTest/Source/ConcertSyncTest/Private/Replication/ReplicationCacheTests.cpp
deleted file mode 100644
