Fix several issues around the client changing streams & authority:
- Taking authority calls IConcertClientReplicationBridge::PushTrackedObjects
- Releasing authority calls IConcertClientReplicationBridge::ReleaseTrackedObjects
- Changing stream info while client is sending now actually updates the properties being replicated

diff --git a/Engine/Plugins/Developer/Concert/ConcertSync/ConcertSyncClient/Source/ConcertSyncClient/Private/Replication/Manager/ReplicationManagerState_Connected.cpp b/Engine/Plugins/Developer/Concert/ConcertSync/ConcertSyncClient/Source/ConcertSyncClient/Private/Replication/Manager/ReplicationManagerState_Connected.cpp
--- a/Engine/Plugins/Developer/Concert/ConcertSync/ConcertSyncClient/Source/ConcertSyncClient/Private/Replication/Manager/ReplicationManagerState_Connected.cpp
+++ b/Engine/Plugins/Developer/Concert/ConcertSync/ConcertSyncClient/Source/ConcertSyncClient/Private/Replication/Manager/ReplicationManagerState_Connected.cpp
@@ -14,6 +14,7 @@
 #include "Replication/Processing/ObjectReplicationReceiver.h"
 #include "Replication/Processing/ObjectReplicationSender.h"
 
+#include "Algo/RemoveIf.h"
 #include "UObject/Package.h"
 #include "UObject/UObjectGlobals.h"
 
@@ -29,9 +30,16 @@ namespace UE::ConcertSyncClient::Replication
 		, LiveSession(LiveSession)
 		, ReplicationBridge(ReplicationBridge)
 		, RegisteredStreams(MoveTemp(StreamDescriptions))
-		// TODO: Use config to determine which replication format to use
+		// TODO DP: Use config to determine which replication format to use
 		, ReplicationFormat(MakeShared<ConcertSyncCore::FFullObjectFormat>())
-		, ReplicationDataSource(MakeShared<FClientReplicationDataCollector>(ReplicationBridge, ReplicationFormat, RegisteredStreams))
+		, ReplicationDataSource(MakeShared<FClientReplicationDataCollector>(
+			ReplicationBridge,
+			ReplicationFormat,
+			FClientReplicationDataCollector::FGetClientStreams::CreateLambda([this]()
+			{
+				return &RegisteredStreams;
+			}))
+			)
 		, Sender(MakeShared<ConcertSyncCore::FObjectReplicationSender>(LiveSession->GetSessionServerEndpointId(), LiveSession, ReplicationDataSource))
 		, ReceivedDataCache(MakeShared<ConcertSyncCore::FObjectReplicationCache>(ReplicationFormat))
 		, Receiver(MakeShared<ConcertSyncCore::FObjectReplicationReceiver>(LiveSession, ReceivedDataCache))
@@ -73,9 +81,18 @@ namespace UE::ConcertSyncClient::Replication
 
 	TFuture<FAuthorityChangeResponse> FReplicationManagerState_Connected::RequestAuthorityChange(FAuthorityChangeRequest Args)
 	{
+		// Stop replicating removed objects right now: the server will remove authority after processing this request.
+		// At that point, it will log errors for receiving replication data from a client without authority.
+		HandleReleasingReplicatedObjects(Args);
+		
 		return LiveSession->SendCustomRequest<FConcertChangeAuthority_Request, FConcertChangeAuthority_Response>(Args, LiveSession->GetSessionServerEndpointId())
-			.Next([](FConcertChangeAuthority_Response&& Response)
+			.Next([WeakThis = TWeakPtr<FReplicationManagerState_Connected>(SharedThis(this)), Args](FConcertChangeAuthority_Response&& Response) mutable
 			{
+				if (const TSharedPtr<FReplicationManagerState_Connected> ThisPin = WeakThis.Pin())
+				{
+					ThisPin->UpdateReplicatedObjectsAfterAuthorityChange(MoveTemp(Args), Response);
+				}
+
 				return FAuthorityChangeResponse { MoveTemp(Response) };
 			});
 	}
@@ -97,14 +114,17 @@ namespace UE::ConcertSyncClient::Replication
 
 	TFuture<FChangeStreamResponse> FReplicationManagerState_Connected::ChangeStream(FChangeStreamRequest Args)
 	{
+		// Stop replicating removed objects right now: the server will remove authority after processing this request.
+		// At that point, it will log errors for receiving replication data from a client without authority.
+		HandleRemovingReplicatedObjects(Args);
+		
 		return LiveSession->SendCustomRequest<FConcertChangeStream_Request, FConcertChangeStream_Response>(Args, LiveSession->GetSessionServerEndpointId())
 			.Next([WeakThis = TWeakPtr<FReplicationManagerState_Connected>(SharedThis(this)), Args](FConcertChangeStream_Response&& Response)
 			{
 				const TSharedPtr<FReplicationManagerState_Connected> ThisPin = WeakThis.Pin();
 				if (ThisPin && Response.IsSuccess())
 				{
-					// Need to update local cache of streams
-					ConcertSyncCore::Replication::ChangeStreamUtils::ApplyValidatedRequest(Args, ThisPin->RegisteredStreams);
+					ThisPin->UpdateReplicatedObjectsAfterStreamChange(Args);
 				}
 				
 				return FChangeStreamResponse { MoveTemp(Response) };
@@ -118,7 +138,7 @@ namespace UE::ConcertSyncClient::Replication
 
 	void FReplicationManagerState_Connected::Tick(IConcertClientSession& Session, float DeltaTime)
 	{
-		// TODO: Set this up in a config file
+		// TODO DP: Set this up in a config file
 		constexpr double TimeBudget = 1.0 / 60.0;
 		double TimeLeft = TimeBudget;
 
@@ -147,5 +167,85 @@ namespace UE::ConcertSyncClient::Replication
 	{
 		ReplicationApplier->ProcessObjects(TimeBudget);
 	}
+
+	void FReplicationManagerState_Connected::UpdateReplicatedObjectsAfterStreamChange(const FChangeStreamRequest& Request)
+	{
+		// Build RegisteredStreams while RegisteredStreams has the old, unupdated state
+		TMap<FSoftObjectPath, TArray<FGuid>> BundledModifiedObjects;
+		for (const TPair<FObjectInStreamID, FConcertChangeStream_PutObject>& PutObject : Request.ObjectsToPut)
+		{
+			const FObjectInStreamID ObjectInfo = PutObject.Key;
+			const FSoftObjectPath Object = ObjectInfo.Object;
+			const FGuid StreamId = ObjectInfo.StreamId;
+			
+			const FReplicationStreamDescription* StreamDescription = RegisteredStreams.FindByPredicate([&StreamId](const FReplicationStreamDescription& Stream)
+			{
+				return Stream.BaseDescription.Identifier == StreamId;
+			});
+			// ReplicationDataSource only cares about inflight objects.
+			// Newly added objects inflight because the client must first request authority.
+			const bool bWasAddedByRequest = !ensure(StreamDescription) || !StreamDescription->BaseDescription.ReplicationMap.ReplicatedObjects.Contains(Object);
+			if (!bWasAddedByRequest)
+			{
+				BundledModifiedObjects.FindOrAdd(Object).Add(StreamId);
+			}
+		}
+
+		// The local cache must be updated before calling OnObjectStreamModified
+		ConcertSyncCore::Replication::ChangeStreamUtils::ApplyValidatedRequest(Request, RegisteredStreams);
+		
+		for (const TPair<FSoftObjectPath, TArray<FGuid>>& ModifiedObject : BundledModifiedObjects)
+		{
+			ReplicationDataSource->OnObjectStreamModified(ModifiedObject.Key, ModifiedObject.Value);
+		}
+	}
+
+	void FReplicationManagerState_Connected::HandleRemovingReplicatedObjects(const FChangeStreamRequest& Request) const
+	{
+		TMap<FSoftObjectPath, TArray<FGuid>> BundledRemovedObjects;
+		for (const FObjectInStreamID& RemovedObject : Request.ObjectsToRemove)
+		{
+			BundledRemovedObjects.FindOrAdd(RemovedObject.Object).Add(RemovedObject.StreamId);
+		}
+		for (const TPair<FSoftObjectPath, TArray<FGuid>>& RemovedObjectInfo : BundledRemovedObjects)
+		{
+			// Removing an object from a stream implies removing its authority so stop replicating it
+			ReplicationDataSource->RemoveReplicatedObjectStreams(RemovedObjectInfo.Key, RemovedObjectInfo.Value);
+		}
+	}
+
+	void FReplicationManagerState_Connected::UpdateReplicatedObjectsAfterAuthorityChange(FAuthorityChangeRequest&& Request, const FConcertChangeAuthority_Response& Response) const
+	{
+		for (TPair<FSoftObjectPath, FConcertStreamArray>& TakeAuthority : Request.TakeAuthority)
+		{
+			const FSoftObjectPath& ReplicatedObject = TakeAuthority.Key;
+			// Request will be discarded so ...
+			FConcertStreamArray& ReplicatedStreams = TakeAuthority.Value;
+
+			const FConcertStreamArray* RejectedStreams = Response.RejectedObjects.Find(ReplicatedObject);
+			if (RejectedStreams)
+			{
+				// ... reuse the memory
+				ReplicatedStreams.StreamIds.SetNum(Algo::RemoveIf(ReplicatedStreams.StreamIds, [RejectedStreams](const FGuid& Stream)
+				{
+					return RejectedStreams->StreamIds.Contains(Stream);
+				}));
+			}
+
+			const bool bWasFullyRejected = ReplicatedStreams.StreamIds.IsEmpty(); 
+			if (!bWasFullyRejected)
+			{
+				ReplicationDataSource->AddReplicatedObjectStreams(TakeAuthority.Key, ReplicatedStreams.StreamIds);
+			}
+		}
+	}
+
+	void FReplicationManagerState_Connected::HandleReleasingReplicatedObjects(const FAuthorityChangeRequest& Request) const
+	{
+		for (const TPair<FSoftObjectPath, FConcertStreamArray>& ReleaseAuthority : Request.ReleaseAuthority)
+		{
+			ReplicationDataSource->RemoveReplicatedObjectStreams(ReleaseAuthority.Key, ReleaseAuthority.Value.StreamIds);
+		}
+	}
 }
 
diff --git a/Engine/Plugins/Developer/Concert/ConcertSync/ConcertSyncClient/Source/ConcertSyncClient/Private/Replication/Manager/ReplicationManagerState_Connected.h b/Engine/Plugins/Developer/Concert/ConcertSync/ConcertSyncClient/Source/ConcertSyncClient/Private/Replication/Manager/ReplicationManagerState_Connected.h
--- a/Engine/Plugins/Developer/Concert/ConcertSync/ConcertSyncClient/Source/ConcertSyncClient/Private/Replication/Manager/ReplicationManagerState_Connected.h
+++ b/Engine/Plugins/Developer/Concert/ConcertSync/ConcertSyncClient/Source/ConcertSyncClient/Private/Replication/Manager/ReplicationManagerState_Connected.h
@@ -115,5 +115,18 @@ namespace UE::ConcertSyncClient::Replication
 		void TickSender(float TimeBudget);
 		/** Processes received data and serializes UObjects. */
 		void TickReceiver(float TimeBudget);
+		
+		/** Updates replicated objects affected by the change request. */
+		void UpdateReplicatedObjectsAfterStreamChange(const FChangeStreamRequest& Request);
+		void HandleRemovingReplicatedObjects(const FChangeStreamRequest& Request) const;
+
+		/**
+		 * Updates the objects which should be replicated after changing authority.
+		 * 
+		 * @note Request is accepted as && because this function rewrites its memory when looking at rejections.
+		 * Since the request was already sent to the server it is assumed the request can just contain trash after.
+		 */
+		void UpdateReplicatedObjectsAfterAuthorityChange(FAuthorityChangeRequest&& Request, const FConcertChangeAuthority_Response& Response) const;
+		void HandleReleasingReplicatedObjects(const FAuthorityChangeRequest& Request) const;
 	};
 }
diff --git a/Engine/Plugins/Developer/Concert/ConcertSync/ConcertSyncClient/Source/ConcertSyncClient/Private/Replication/Processing/ClientReplicationDataCollector.cpp b/Engine/Plugins/Developer/Concert/ConcertSync/ConcertSyncClient/Source/ConcertSyncClient/Private/Replication/Processing/ClientReplicationDataCollector.cpp
--- a/Engine/Plugins/Developer/Concert/ConcertSync/ConcertSyncClient/Source/ConcertSyncClient/Private/Replication/Processing/ClientReplicationDataCollector.cpp
+++ b/Engine/Plugins/Developer/Concert/ConcertSync/ConcertSyncClient/Source/ConcertSyncClient/Private/Replication/Processing/ClientReplicationDataCollector.cpp
@@ -2,6 +2,7 @@
 
 #include "ClientReplicationDataCollector.h"
 
+#include "Algo/RemoveIf.h"
 #include "Replication/Data/ReplicationStreamDescription.h"
 #include "Replication/Formats/IObjectReplicationFormat.h"
 #include "Replication/IConcertClientReplicationBridge.h"
@@ -11,36 +12,17 @@
 namespace UE::ConcertSyncClient::Replication
 {
 	FClientReplicationDataCollector::FClientReplicationDataCollector(
-		IConcertClientReplicationBridge* ReplicationBridge,
-		TSharedRef<ConcertSyncCore::IObjectReplicationFormat> ReplicationFormat,
-		TArrayView<const FReplicationStreamDescription> StreamDescriptions
+		IConcertClientReplicationBridge* InReplicationBridge,
+		TSharedRef<ConcertSyncCore::IObjectReplicationFormat> InReplicationFormat,
+		FGetClientStreams InGetStreamsDelegate
 		)
-		: Bridge(ReplicationBridge)
-		, ReplicationFormat(MoveTemp(ReplicationFormat))
+		: Bridge(InReplicationBridge)
+		, ReplicationFormat(MoveTemp(InReplicationFormat))
+		, GetStreamsDelegate(MoveTemp(InGetStreamsDelegate))
 	{
+		check(GetStreamsDelegate.IsBound());
 		Bridge->OnObjectDiscovered().AddRaw(this, &FClientReplicationDataCollector::StartTrackingObject);
 		Bridge->OnObjectHidden().AddRaw(this, &FClientReplicationDataCollector::StopTrackingObject);
-
-		TSet<FSoftObjectPath> PushedSet;
-		for (const FReplicationStreamDescription& StreamDescription : StreamDescriptions)
-		{
-			const FGuid StreamId = StreamDescription.BaseDescription.Identifier;
-			for (const TPair<FSoftObjectPath, FReplicatedObjectInfo>& Pair : StreamDescription.BaseDescription.ReplicationMap.ReplicatedObjects)
-			{
-				const FSoftObjectPath& ObjectPath = Pair.Key;
-				
-				TArray<FObjectInfo>& ExistingObjectInfo = ObjectsToReplicate.FindOrAdd(ObjectPath);
-				checkf(!ExistingObjectInfo.ContainsByPredicate([&StreamId](const FObjectInfo& ObjectInfo){ return ObjectInfo.StreamId == StreamId; }), TEXT("Invalid data should have been rejected during handshake"));
-
-				ExistingObjectInfo.Add({ StreamId, Pair.Value.PropertySelection });
-				// Only push it once because each push increments a counter
-				if (!PushedSet.Contains(ObjectPath))
-				{
-					PushedSet.Add(ObjectPath);
-					Bridge->PushTrackedObjects({ ObjectPath });
-				}
-			}
-		}
 	}
 
 	FClientReplicationDataCollector::~FClientReplicationDataCollector()
@@ -53,6 +35,94 @@ namespace UE::ConcertSyncClient::Replication
 			Bridge->PopTrackedObjects({ Pair.Key });
 		}
 	}
+	
+	void FClientReplicationDataCollector::AddReplicatedObjectStreams(const FSoftObjectPath& Object, TArrayView<const FGuid> AddedStreams)
+	{
+		TArray<FObjectInfo>& ReplicatedObjectInfo = ObjectsToReplicate.FindOrAdd(Object);
+		const bool bIsNewReplicatedObject = ReplicatedObjectInfo.IsEmpty();
+		
+		const TArray<FReplicationStreamDescription>& RegisteredStreams = *GetStreamsDelegate.Execute();
+		ReplicatedObjectInfo.Reserve(ReplicatedObjectInfo.Num() + AddedStreams.Num());
+		for (const FGuid& StreamId : AddedStreams)
+		{
+			const FReplicationStreamDescription* Stream = RegisteredStreams.FindByPredicate([&StreamId](const FReplicationStreamDescription& Stream)
+			{
+				return Stream.BaseDescription.Identifier == StreamId;
+			});
+			const FReplicatedObjectInfo* ObjectInfo = Stream ? Stream->BaseDescription.ReplicationMap.ReplicatedObjects.Find(Object) : nullptr;
+			if (ensureAlwaysMsgf(ObjectInfo, TEXT("Client's registered streams cache is out of sync")))
+			{
+				ReplicatedObjectInfo.Add({ StreamId, ObjectInfo->PropertySelection });
+			}
+		}
+		
+		// FClientReplicationDataCollector should not push Object more than once because each push increments an internal counter
+		if (bIsNewReplicatedObject)
+		{
+			Bridge->PushTrackedObjects({ Object });
+		}
+	}
+
+	void FClientReplicationDataCollector::RemoveReplicatedObjectStreams(const FSoftObjectPath& Object, TArrayView<const FGuid> RemovedStreams)
+	{
+		TArray<FObjectInfo>* ReplicatedObjectInfo = ObjectsToReplicate.Find(Object);
+		if (!ReplicatedObjectInfo)
+		{
+			// This object is not being replicated
+			return;
+		}
+
+		ReplicatedObjectInfo->SetNum(Algo::RemoveIf(*ReplicatedObjectInfo, [&RemovedStreams](const FObjectInfo& ObjectInfo)
+		{
+			return RemovedStreams.Contains(ObjectInfo.StreamId);
+		}));
+
+		if (ReplicatedObjectInfo->IsEmpty())
+		{
+			ObjectsToReplicate.Remove(Object);
+			Bridge->PopTrackedObjects({ Object });
+		}
+	}
+
+	void FClientReplicationDataCollector::OnObjectStreamModified(const FSoftObjectPath& Object, TArrayView<const FGuid> PutStreams)
+	{
+		TArray<FObjectInfo>* ReplicatedObjectInfo = ObjectsToReplicate.Find(Object);
+		if (!ReplicatedObjectInfo)
+		{
+			// This object is not being replicated
+			return;
+		}
+
+		const TArray<FReplicationStreamDescription>& RegisteredStreams = *GetStreamsDelegate.Execute();
+		for (const FGuid& StreamId : PutStreams)
+		{
+			const FReplicationStreamDescription* Stream = RegisteredStreams.FindByPredicate([&StreamId](const FReplicationStreamDescription& Stream)
+			{
+				return Stream.BaseDescription.Identifier == StreamId;
+			});
+			const FReplicatedObjectInfo* ObjectInfo = Stream ? Stream->BaseDescription.ReplicationMap.ReplicatedObjects.Find(Object) : nullptr;
+			if (!ensureAlwaysMsgf(ObjectInfo, TEXT("Client's registered streams cache is out of sync")))
+			{
+				continue;
+			}
+
+			// The PutObject request either just created ObjectInfo or updated it.
+			FObjectInfo* ReplicatedObject = ReplicatedObjectInfo->FindByPredicate([&](const FObjectInfo& Existing)
+			{
+				return Existing.StreamId == StreamId;
+			});
+			if (ReplicatedObject)
+			{
+				// PutObject updated existing object in stream
+				ReplicatedObject->SelectedProperties = ObjectInfo->PropertySelection;
+			}
+			else
+			{
+				// PutObject added object to stream
+				ReplicatedObjectInfo->Add({ StreamId, ObjectInfo->PropertySelection });
+			}
+		}
+	}
 
 	void FClientReplicationDataCollector::ForEachPendingObject(TFunctionRef<void(const FObjectInStreamID&)> ProcessItemFunc) const
 	{
diff --git a/Engine/Plugins/Developer/Concert/ConcertSync/ConcertSyncClient/Source/ConcertSyncClient/Private/Replication/Processing/ClientReplicationDataCollector.h b/Engine/Plugins/Developer/Concert/ConcertSync/ConcertSyncClient/Source/ConcertSyncClient/Private/Replication/Processing/ClientReplicationDataCollector.h
--- a/Engine/Plugins/Developer/Concert/ConcertSync/ConcertSyncClient/Source/ConcertSyncClient/Private/Replication/Processing/ClientReplicationDataCollector.h
+++ b/Engine/Plugins/Developer/Concert/ConcertSync/ConcertSyncClient/Source/ConcertSyncClient/Private/Replication/Processing/ClientReplicationDataCollector.h
@@ -7,6 +7,7 @@
 #include "Containers/ArrayView.h"
 #include "Containers/Set.h"
 #include "Replication/Data/ConcertPropertySelection.h"
+#include "Replication/Data/ObjectIds.h"
 #include "Templates/SharedPointer.h"
 #include "UObject/SoftObjectPtr.h"
 
@@ -24,7 +25,6 @@ namespace UE::ConcertSyncCore
 
 namespace UE::ConcertSyncClient::Replication
 {
-	
 	/**
 	 * Exposes UObject instances to an FObjectReplicationProcessor.
 	 * IConcertClientReplicationBridge tracks UObject lifetime, this class exposes them.
@@ -32,13 +32,33 @@ namespace UE::ConcertSyncClient::Replication
 	class FClientReplicationDataCollector : public ConcertSyncCore::IReplicationDataSource
 	{
 	public:
-		
+
+		DECLARE_DELEGATE_RetVal(const TArray<FReplicationStreamDescription>*, FGetClientStreams);
 		FClientReplicationDataCollector(
-			IConcertClientReplicationBridge* ReplicationBridge,
-			TSharedRef<ConcertSyncCore::IObjectReplicationFormat> ReplicationFormat,
-			TArrayView<const FReplicationStreamDescription> StreamsToSend
+			IConcertClientReplicationBridge* InReplicationBridge,
+			TSharedRef<ConcertSyncCore::IObjectReplicationFormat> InReplicationFormat,
+			FGetClientStreams InGetStreamsDelegate
 			);
 		virtual ~FClientReplicationDataCollector() override;
+		
+		/**
+		 * Indicates that this object should start replicating for the given streams.
+		 * @param Object The object that should start replicating
+		 * @param AddedStreams The streams determine which properties are to be replicated
+		 */
+		void AddReplicatedObjectStreams(const FSoftObjectPath& Object, TArrayView<const FGuid> AddedStreams);
+		/**
+		 * Indicates that certain properties of an object should no longer be replicated.
+		 * @param Object The object that should start replicating
+		 * @param RemovedStreams The streams determine which properties are to be replicated
+		 */
+		void RemoveReplicatedObjectStreams(const FSoftObjectPath& Object, TArrayView<const FGuid> RemovedStreams);
+		/**
+		 * Called when the client modifies pre-existing object. Adjusts any inflight replication if needed.
+		 * @param Object The object that should start replicating
+		 * @param PutStreams The streams determine which properties are to be replicated
+		 */
+		void OnObjectStreamModified(const FSoftObjectPath& Object, TArrayView<const FGuid> PutStreams);
 
 		//~ Begin IReplicationDataSource Interface
 		virtual void ForEachPendingObject(TFunctionRef<void(const FObjectInStreamID&)> ProcessItemFunc) const override;
@@ -51,7 +71,10 @@ namespace UE::ConcertSyncClient::Replication
 		/** Gets and tracks replicated objects */
 		IConcertClientReplicationBridge* Bridge;
 		/** Used to create the replication data sent to the server. */
-		TSharedRef<ConcertSyncCore::IObjectReplicationFormat> ReplicationFormat;
+		const TSharedRef<ConcertSyncCore::IObjectReplicationFormat> ReplicationFormat;
+
+		/** Gets the stream of the managed client. */
+		const FGetClientStreams GetStreamsDelegate;
 
 		struct FObjectInfo
 		{
diff --git a/Engine/Plugins/Developer/Concert/ConcertSync/ConcertSyncCore/Source/ConcertSyncCore/Private/Replication/Processing/ObjectReplicationReceiver.cpp b/Engine/Plugins/Developer/Concert/ConcertSync/ConcertSyncCore/Source/ConcertSyncCore/Private/Replication/Processing/ObjectReplicationReceiver.cpp
--- a/Engine/Plugins/Developer/Concert/ConcertSync/ConcertSyncCore/Source/ConcertSyncCore/Private/Replication/Processing/ObjectReplicationReceiver.cpp
+++ b/Engine/Plugins/Developer/Concert/ConcertSync/ConcertSyncCore/Source/ConcertSyncCore/Private/Replication/Processing/ObjectReplicationReceiver.cpp
@@ -52,14 +52,17 @@ namespace UE::ConcertSyncCore
 				}
 			}
 		}
-		
-		UE_CLOG(CVarLogReceivedObjects.GetValueOnGameThread(), LogConcert, Log, TEXT("Received %d streams with %d object changes from endpoint %s. Rejected %d object changes. Cached %d object changes with a total of %d cache usages."),
-			Event.Streams.Num(),
-			NumObjects,
-			*SessionContext.SourceEndpointId.ToString(),
-			NumRejectedObjectChanges,
-			NumOfAcceptedObjectChanges,
-			NumCacheUsages
+
+		if (CVarLogReceivedObjects.GetValueOnGameThread())
+		{
+			UE_LOG(LogConcert, Log, TEXT("Received %d streams with %d object changes from endpoint %s. Cached %d object changes with a total of %d cache usages."),
+				Event.Streams.Num(),
+				NumObjects,
+				*SessionContext.SourceEndpointId.ToString(),
+				NumOfAcceptedObjectChanges,
+				NumCacheUsages
 			);
+			UE_CLOG(NumRejectedObjectChanges > 0, LogConcert, Warning, TEXT("Rejected %d object changes."), NumRejectedObjectChanges);
+		}
 	}
 }
diff --git a/Engine/Plugins/Developer/Concert/ConcertSync/ConcertSyncCore/Source/ConcertSyncCore/Public/Replication/Messages/ConcertReplicationEvents.h b/Engine/Plugins/Developer/Concert/ConcertSync/ConcertSyncCore/Source/ConcertSyncCore/Public/Replication/Messages/ConcertReplicationEvents.h
--- a/Engine/Plugins/Developer/Concert/ConcertSync/ConcertSyncCore/Source/ConcertSyncCore/Public/Replication/Messages/ConcertReplicationEvents.h
+++ b/Engine/Plugins/Developer/Concert/ConcertSync/ConcertSyncCore/Source/ConcertSyncCore/Public/Replication/Messages/ConcertReplicationEvents.h
@@ -329,6 +329,8 @@ struct FConcertChangeStream_Response
 	bool IsSuccess() const { return AuthorityConflicts.IsEmpty() && ObjectsToPutSemanticErrors.IsEmpty() && FailedStreamCreation.IsEmpty(); }
 	bool IsFailure() const { return !IsSuccess(); }
 
+	bool WasObjectPutSuccessful(const FObjectInStreamID& Object) const { return !AuthorityConflicts.Contains(Object) && !ObjectsToPutSemanticErrors.Contains(Object); }
+
 	/** If IsFailure(), logs the errors. */
 	CONCERTSYNCCORE_API void LogErrors(FOutputDevice& OutputDevice) const;
 };
\ No newline at end of file
diff --git a/Engine/Plugins/Developer/Concert/ConcertSync/ConcertSyncServer/Source/ConcertSyncServer/Private/Replication/AuthorityManager.cpp b/Engine/Plugins/Developer/Concert/ConcertSync/ConcertSyncServer/Source/ConcertSyncServer/Private/Replication/AuthorityManager.cpp
--- a/Engine/Plugins/Developer/Concert/ConcertSync/ConcertSyncServer/Source/ConcertSyncServer/Private/Replication/AuthorityManager.cpp
+++ b/Engine/Plugins/Developer/Concert/ConcertSync/ConcertSyncServer/Source/ConcertSyncServer/Private/Replication/AuthorityManager.cpp
@@ -164,22 +164,7 @@ namespace UE::ConcertSyncServer::Replication
 			else
 			{
 				UE_LOG(LogConcert, Log, TEXT("Rejected %s request of authority over %s in stream %s"), *ClientId.ToString(EGuidFormats::Short), *ObjectPath.ToString(), *StreamId.ToString(EGuidFormats::Short));
-				Response.RejectedObjects.FindOrAdd(ObjectPath).StreamIds.Add(ClientId);
-			}
-		});
-		
-		Private::ForEachReplicatedObject(Request.TakeAuthority, [this, &Response, &AuthorityData, &ClientId](const FStreamId& StreamId, const FSoftObjectPath& ObjectPath)
-		{
-			const FReplicatedObjectId ObjectToAuthor{ { StreamId, ObjectPath }, ClientId };
-			if (CanTakeAuthority(ObjectToAuthor))
-			{
-				UE_LOG(LogConcert, Log, TEXT("Transferred authority of %s to client %s for their stream %s"), *ObjectPath.ToString(), *ClientId.ToString(EGuidFormats::Short), *StreamId.ToString(EGuidFormats::Short));
-				AuthorityData.OwnedObjects.FindOrAdd(StreamId).Add(ObjectPath);
-			}
-			else
-			{
-				UE_LOG(LogConcert, Log, TEXT("Rejected %s request of authority over %s in stream %s"), *ClientId.ToString(EGuidFormats::Short), *ObjectPath.ToString(), *StreamId.ToString(EGuidFormats::Short));
-				Response.RejectedObjects.Add(ObjectPath);
+				Response.RejectedObjects.FindOrAdd(ObjectPath).StreamIds.AddUnique(StreamId);
 			}
 		});
 		
diff --git a/Engine/Plugins/Developer/Concert/ConcertSync/ConcertSyncTest/Source/ConcertSyncTest/Private/Replication/ReplicationAuthorityTests.cpp b/Engine/Plugins/Developer/Concert/ConcertSync/ConcertSyncTest/Source/ConcertSyncTest/Private/Replication/ReplicationAuthorityTests.cpp
--- a/Engine/Plugins/Developer/Concert/ConcertSync/ConcertSyncTest/Source/ConcertSyncTest/Private/Replication/ReplicationAuthorityTests.cpp
+++ b/Engine/Plugins/Developer/Concert/ConcertSync/ConcertSyncTest/Source/ConcertSyncTest/Private/Replication/ReplicationAuthorityTests.cpp
@@ -9,6 +9,8 @@
 #include "Misc/AutomationTest.h"
 #include "Misc/Paths.h"
 #include "HAL/FileManager.h"
+#include "Replication/Formats/FullObjectFormat.h"
+#include "Util/ConcertClientReplicationBridgeMock.h"
 #include "Util/SendReceiveGenericStreamTestBase.h"
 
 namespace UE::ConcertSyncTests::Replication::Authority
@@ -20,7 +22,28 @@ namespace UE::ConcertSyncTests::Replication::Authority
 		// 1. Init
 		SetUpClientAndServer();
 		
-		// 2. Send data without authority > Reject
+		// 2.1 Generate replication data
+		ConcertSyncCore::FFullObjectFormat SerializeObject;
+		const TOptional<FConcertSessionSerializedPayload> Payload = SerializeObject.CreateReplicationEvent(*TestObject, [](const FArchiveSerializedPropertyChain* Chain, const FProperty& Property)
+		{
+			return true;
+		});
+		if (!Payload)
+		{
+			AddError(TEXT("Faild to create payload"));
+			return false;
+		}
+		
+		FConcertObjectReplicationEvent ObjectReplicationEvent;
+		ObjectReplicationEvent.ReplicatedObject = TestObject;
+		ObjectReplicationEvent.SerializedPayload = *Payload;
+		FConcertStreamReplicationEvent ReplicationEvent;
+		ReplicationEvent.StreamId = SenderStreamId;
+		ReplicationEvent.ReplicatedObjects.Add(ObjectReplicationEvent);
+		FConcertBatchReplicationEvent ReplicationBatchEvent;
+		ReplicationBatchEvent.Streams.Add(ReplicationEvent);
+		
+		// 2.2 Prepare for sending data without authority > Reject
 		bool bHasServerReceivedData = false;
 		auto OnServerReceive = [this, &bHasServerReceivedData](const FConcertSessionContext& Context, const FConcertBatchReplicationEvent& Event) mutable
 		{
@@ -34,8 +57,34 @@ namespace UE::ConcertSyncTests::Replication::Authority
 		{
 			AddError(TEXT("Server sent data from non-authorative client to receiving client!"));
 		};
-		SimulateSenderToReceiver(OnServerReceive, OnClientReceive);
-		TestTrue(TEXT("Server received replication event from non-authorative client"), bHasServerReceivedData);
+		ServerSession->RegisterCustomEventHandler<FConcertBatchReplicationEvent>(OnServerReceive);
+		Client_Receiver->ClientSessionMock->RegisterCustomEventHandler<FConcertBatchReplicationEvent>(OnClientReceive);
+		const FGuid ServerSessionId { 0, 0, 0, 0};
+		
+		// 3. Test that server processed the data and rejected it.
+		IConsoleVariable* ConsoleVariable = IConsoleManager::Get().FindConsoleVariable(TEXT("Concert.Replication.LogReceivedObjects"));
+		if (ConsoleVariable)
+		{
+			// Hacky way of making sure that the server REALLY rejected the change...
+			// In the future we should add a callback that can be registered via the server instance.
+			AddExpectedError(TEXT("Rejected 1 object change"));
+			bool bValue;
+			ConsoleVariable->GetValue(bValue);
+			ConsoleVariable->Set(true);
+			
+			// Send data without authority > Reject
+			Client_Sender->ClientSessionMock->SendCustomEvent(ReplicationBatchEvent, { ServerSessionId }, EConcertMessageFlags::ReliableOrdered);
+			TickServer();
+			TestTrue(TEXT("Server received replication event from non-authorative client"), bHasServerReceivedData);
+			
+			ConsoleVariable->Set(bValue);
+		}
+		else
+		{
+			Client_Sender->ClientSessionMock->SendCustomEvent(ReplicationBatchEvent, { ServerSessionId }, EConcertMessageFlags::ReliableOrdered);
+			TickServer();
+			TestTrue(TEXT("Server received replication event from non-authorative client"), bHasServerReceivedData);
+		}
 		
 		return true;
 	}
@@ -47,7 +96,10 @@ namespace UE::ConcertSyncTests::Replication::Authority
 		// 1. Init
 		const UObject* UnusedTestObject = GetDefault<UTestReflectionObject>();
 		const FSoftObjectPath TestObjectPath { UnusedTestObject };
-		SenderArgs = CreateHandshakeArgsFrom(*UnusedTestObject);
+		
+		const FGuid SenderStreamId = FGuid::NewGuid();
+		const FGuid ReceiverStreamId = SenderStreamId;
+		SenderArgs = CreateHandshakeArgsFrom(*UnusedTestObject, SenderStreamId);
 		ReceiverArgs = SenderArgs;
 		
 		SetUpClientAndServer();
@@ -64,13 +116,23 @@ namespace UE::ConcertSyncTests::Replication::Authority
 		TickServer();
 		TestTrue(TEXT("Sender received response to take authority"), bSenderReceivedResponse);
 
-		// 2.2 ... so does receiver fails taking authority over the same object
+		// 2.2 ... so does receiver, who fails taking authority over the same object
 		bool bReceiverReceivedResponse = false;
 		ClientReplicationManager_Receiver->TakeAuthorityOver({ TestObjectPath })
-			.Next([this, &bReceiverReceivedResponse](const ConcertSyncClient::Replication::FAuthorityChangeResponse& Response) mutable
+			.Next([this, &TestObjectPath, &ReceiverStreamId, &bReceiverReceivedResponse](const ConcertSyncClient::Replication::FAuthorityChangeResponse& Response) mutable
 			{
 				bReceiverReceivedResponse = true;
-				TestEqual(TEXT("Rejected because Sender already has authority"), Response.RejectedObjects.Num(), 1); 
+				TestEqual(TEXT("Rejected because Sender already has authority"), Response.RejectedObjects.Num(), 1);
+				
+				if (const FConcertStreamArray* RejectedStreams = Response.RejectedObjects.Find(TestObjectPath))
+				{
+					TestEqual(TEXT("1 stream rejected"), RejectedStreams->StreamIds.Num(), 1);
+					TestTrue(TEXT("Rejected streams contains stream registered by receiver client"), RejectedStreams->StreamIds.Contains(ReceiverStreamId));
+				}
+				else
+				{
+					AddError(TEXT("Expected rejected stream array for TestObjectPath"));
+				}
 			});
 		TickClient(Client_Receiver);
 		TickServer();
@@ -115,15 +177,26 @@ namespace UE::ConcertSyncTests::Replication::Authority
 		// 2. Attempt to take authority over object that was not in handshake
 		bool bReceivedResponse = false;
 		const FSoftObjectPath SomePath(GetDefault<UTestReflectionObject>());
+		const FGuid StreamId = FGuid::NewGuid();
 		ConcertSyncClient::Replication::FAuthorityChangeRequest Request;
-		Request.TakeAuthority.Add(SomePath, FConcertStreamArray{ .StreamIds = { FGuid::NewGuid() } });
+		Request.TakeAuthority.Add(SomePath, FConcertStreamArray{ .StreamIds = { StreamId } });
 		// Detail: Cannot use IConcertClientReplicationManager::TakeAuthority util here because it builds the request based on what streams were registered
 		ClientReplicationManager_Sender->RequestAuthorityChange(Request)
-			.Next([this, &bReceivedResponse, SomePath](const ConcertSyncClient::Replication::FAuthorityChangeResponse& Response) mutable
+			.Next([this, &bReceivedResponse, &SomePath, &StreamId](const ConcertSyncClient::Replication::FAuthorityChangeResponse& Response) mutable
 			{
 				bReceivedResponse = true;
 				TestTrue(TEXT("Cannot take authority over unregistered object"), Response.RejectedObjects.Contains(SomePath));
-				TestEqual(TEXT("Rejected exactly 1 object"), Response.RejectedObjects.Num(), 1); 
+				TestEqual(TEXT("Rejected exactly 1 object"), Response.RejectedObjects.Num(), 1);
+
+				if (const FConcertStreamArray* RejectedStreams = Response.RejectedObjects.Find(SomePath))
+				{
+					TestEqual(TEXT("1 stream rejected"), RejectedStreams->StreamIds.Num(), 1);
+					TestTrue(TEXT("Rejected streams contains stream registered by sender client"), RejectedStreams->StreamIds.Contains(StreamId));
+				}
+				else
+				{
+					AddError(TEXT("Expected rejected stream array for SomePath"));
+				}
 			});
 
 		// 3. Authority request was rejected
@@ -220,4 +293,46 @@ namespace UE::ConcertSyncTests::Replication::Authority
 		
 		return true;
 	}
+	
+	/**
+	 * Tests that the client's local bridge is updated if and only if changing authority:
+	 * - Taking authority calls IConcertClientReplicationBridge::PushTrackedObjects
+	 * - Releasing authority calls IConcertClientReplicationBridge::ReleaseTrackedObjects
+	 */
+	IMPLEMENT_CUSTOM_SIMPLE_AUTOMATION_TEST(FChangingAuthorityUpdatesClientBridge, FSendReceiveObjectTestBase, "Concert.Replication.Authority.ChangingAuthorityUpdatesClientBridge", EAutomationTestFlags::EditorContext | EAutomationTestFlags::EngineFilter);
+	bool FChangingAuthorityUpdatesClientBridge::RunTest(const FString& Parameters)
+	{
+		// 1. Init
+		SetUpClientAndServer();
+
+		// 2. Send authority
+		TestEqual(TEXT("No tracked objects"), BridgeMock_Sender->TrackedObjects.Num(), 0);
+		bool bSenderReceivedTakeResponse = false;
+		ClientReplicationManager_Sender->TakeAuthorityOver({ TestObject })
+			.Next([this, &bSenderReceivedTakeResponse](const ConcertSyncClient::Replication::FAuthorityChangeResponse& Response) mutable
+			{
+				bSenderReceivedTakeResponse = true;
+				TestEqual(TEXT("No rejection taking authority"), Response.RejectedObjects.Num(), 0); 
+			});
+		TestTrue(TEXT("Sender received response taking authority"), bSenderReceivedTakeResponse);
+
+		// 3. Test: Added TestObject to tracked objects
+		TestEqual(TEXT("Tracked exactly 1 object"), BridgeMock_Sender->TrackedObjects.Num(), 1);
+		TestTrue(TEXT("Tracked TestObject"), BridgeMock_Sender->TrackedObjects.Contains(TestObject));
+
+		// 4. Release authority
+		bool bSenderReceivedReleaseResponse = false;
+		ClientReplicationManager_Sender->ReleaseAuthorityOf({ TestObject })
+			.Next([this, &bSenderReceivedReleaseResponse](const ConcertSyncClient::Replication::FAuthorityChangeResponse& Response) mutable
+			{
+				bSenderReceivedReleaseResponse = true;
+				TestEqual(TEXT("No rejection releasing authority"), Response.RejectedObjects.Num(), 0); 
+			});
+		TestTrue(TEXT("Sender received response releasing authority"), bSenderReceivedReleaseResponse);
+		
+		// 4. Test: Removed TestObject from tracked objects
+		TestEqual(TEXT("Removed tracked object"), BridgeMock_Sender->TrackedObjects.Num(), 0);
+
+		return true;
+	}
 }
diff --git a/Engine/Plugins/Developer/Concert/ConcertSync/ConcertSyncTest/Source/ConcertSyncTest/Private/Replication/ReplicationSendReceiveTests.cpp b/Engine/Plugins/Developer/Concert/ConcertSync/ConcertSyncTest/Source/ConcertSyncTest/Private/Replication/ReplicationSendReceiveTests.cpp
--- a/Engine/Plugins/Developer/Concert/ConcertSync/ConcertSyncTest/Source/ConcertSyncTest/Private/Replication/ReplicationSendReceiveTests.cpp
+++ b/Engine/Plugins/Developer/Concert/ConcertSync/ConcertSyncTest/Source/ConcertSyncTest/Private/Replication/ReplicationSendReceiveTests.cpp
@@ -22,11 +22,10 @@
 
 namespace UE::ConcertSyncTests::Replication::SendReceiveFlow
 {
-	static bool SharedRunTest(FSendReceiveObjectTestBase& Test)
+	/** Sets up clients, server, and makes the sender take authority over the test object. */
+	static void SharedSetupSet(FSendReceiveObjectTestBase& Test)
 	{
-		// 1. Init
 		Test.SetUpClientAndServer();
-		
 		Test.ClientReplicationManager_Sender->TakeAuthorityOver({ Test.TestObject })
 			.Next([&Test](ConcertSyncClient::Replication::FAuthorityChangeResponse&& Response)
 			{
@@ -35,8 +34,11 @@ namespace UE::ConcertSyncTests::Replication::SendReceiveFlow
 					Test.AddError(TEXT("Failed to take authority"));
 				}
 			});
-		
-		// 2. Send data
+	}
+	
+	/** Send values from sender > server > receiver and validate it arrived. */
+	static void SharedRunSendReceiveTest(FSendReceiveObjectTestBase& Test, EPropertyTestFlags PropertyTestFlags = EPropertyTestFlags::All)
+	{
 		bool bHasServerReceivedData = false;
 		bool bHasClientReceivedData = false;
 		auto OnServerReceive = [&Test, &bHasServerReceivedData](const FConcertSessionContext& Context, const FConcertBatchReplicationEvent& Event) mutable
@@ -55,14 +57,34 @@ namespace UE::ConcertSyncTests::Replication::SendReceiveFlow
 			}
 			bHasClientReceivedData = true;
 		};
-		Test.SimulateSendObjectToReceiver(OnServerReceive, OnClientReceive);
+		Test.SimulateSendObjectToReceiver(OnServerReceive, OnClientReceive,
+			// Explicitly sets all values before & after sending so the caller can test whether the other properties were changed
+			EPropertyTestFlags::All
+			);
 
 		// 3. Test
 		Test.TestTrue(TEXT("Server received replication event"), bHasServerReceivedData);
 		Test.TestTrue(TEXT("Client 2 received replication event"), bHasClientReceivedData);
-		Test.TestEqualTestValues(*Test.TestObject, Test);
+		Test.TestEqualTestValues(*Test.TestObject, PropertyTestFlags);
+	}
+
+	/** Changes the sending stream so it sends only the Float property. */
+	static void ChangeStreamToSendingVectorOnly(FSendReceiveObjectTestBase& Test)
+	{
+		FConcertPropertySelection VectorOnlySelection;
+		const TOptional<FConcertPropertyChain> VectorPropertyChain = FConcertPropertyChain::CreateFromPath(*UTestReflectionObject::StaticClass(), { TEXT("Vector") });
+		VectorOnlySelection.ReplicatedProperties.Add(*VectorPropertyChain);
 		
-		return true;
+		ConcertSyncClient::Replication::FChangeStreamRequest Request;
+		Request.ObjectsToPut.Add(FObjectInStreamID{ Test.SenderStreamId, Test.TestObject }, FConcertChangeStream_PutObject{ VectorOnlySelection });
+		bool bReceivedChangeStreamResponse = false;
+		Test.ClientReplicationManager_Sender->ChangeStream(Request)
+			.Next([&Test, &bReceivedChangeStreamResponse](ConcertSyncClient::Replication::FChangeStreamResponse&& Response)
+			{
+				bReceivedChangeStreamResponse = true;
+				Test.TestTrue(TEXT("Changed Stream"), Response.IsSuccess());
+			});
+		Test.TestTrue(TEXT("Received change stream response"), bReceivedChangeStreamResponse);
 	}
 	
 	/**
@@ -71,7 +93,9 @@ namespace UE::ConcertSyncTests::Replication::SendReceiveFlow
 	IMPLEMENT_CUSTOM_SIMPLE_AUTOMATION_TEST(FSendReceiveFlowTests, FSendReceiveObjectTestBase, "Concert.Replication.SendReceive.SingleStream", EAutomationTestFlags::EditorContext | EAutomationTestFlags::EngineFilter);
 	bool FSendReceiveFlowTests::RunTest(const FString& Parameters)
 	{
-		return SharedRunTest(*this);
+		SharedSetupSet(*this);
+		SharedRunSendReceiveTest(*this);
+		return true;
 	}
 
 	/** Test which still sends TestObject but does so with two separate streams: one for the float and the other for the vector property. */
@@ -118,6 +142,24 @@ namespace UE::ConcertSyncTests::Replication::SendReceiveFlow
 	IMPLEMENT_CUSTOM_SIMPLE_AUTOMATION_TEST(FSendReceiveMultiStreamSameObjectTests, FSplitSendReceiveObjectTest, "Concert.Replication.SendReceive.MultipleStreamsForSameObject", EAutomationTestFlags::EditorContext | EAutomationTestFlags::EngineFilter);
 	bool FSendReceiveMultiStreamSameObjectTests::RunTest(const FString& Parameters)
 	{
-		return SharedRunTest(*this);
+		SharedSetupSet(*this);
+		SharedRunSendReceiveTest(*this);
+		return true;
+	}
+
+	/** Tests that changing a stream while replication is in progress actually changes the properties being replicated. */
+	IMPLEMENT_CUSTOM_SIMPLE_AUTOMATION_TEST(FUpdateStreamWhileSendingUpdatesSentProperties, FSendReceiveObjectTestBase, "Concert.Replication.SendReceive.UpdateStreamWhileSendingUpdatesSentProperties", EAutomationTestFlags::EditorContext | EAutomationTestFlags::EngineFilter);
+	bool FUpdateStreamWhileSendingUpdatesSentProperties::RunTest(const FString& Parameters)
+	{
+		// 1. Send all properties with the stream from handshake
+		SharedSetupSet(*this);
+		SharedRunSendReceiveTest(*this);
+
+		// 2. Change stream to only send floats and test only the float is sent.
+		ChangeStreamToSendingVectorOnly(*this);
+		SharedRunSendReceiveTest(*this, EPropertyTestFlags::Vector);
+		// Test that only the vector was written
+		TestEqualDifferentValues(*TestObject, EPropertyTestFlags::Float);
+		return true;
 	}
 }
diff --git a/Engine/Plugins/Developer/Concert/ConcertSync/ConcertSyncTest/Source/ConcertSyncTest/Private/Replication/Util/ConcertClientReplicationBridgeMock.h b/Engine/Plugins/Developer/Concert/ConcertSync/ConcertSyncTest/Source/ConcertSyncTest/Private/Replication/Util/ConcertClientReplicationBridgeMock.h
--- a/Engine/Plugins/Developer/Concert/ConcertSync/ConcertSyncTest/Source/ConcertSyncTest/Private/Replication/Util/ConcertClientReplicationBridgeMock.h
+++ b/Engine/Plugins/Developer/Concert/ConcertSync/ConcertSyncTest/Source/ConcertSyncTest/Private/Replication/Util/ConcertClientReplicationBridgeMock.h
@@ -28,6 +28,15 @@ namespace UE::ConcertSyncTests::Replication
 			for (const FSoftObjectPath& Path : InTrackedObjects)
 			{
 				TrackedObjects.Add(Path);
+
+				UObject** AvailableObject = AvailableObjects.FindByPredicate([&Path](UObject* Object)
+				{
+					return Path == Object;
+				});
+				if (AvailableObject)
+				{
+					OnObjectDiscoveredDelegate.Broadcast(**AvailableObject);
+				}
 			}
 		}
 		
diff --git a/Engine/Plugins/Developer/Concert/ConcertSync/ConcertSyncTest/Source/ConcertSyncTest/Private/Replication/Util/SendReceiveObjectTestBase.cpp b/Engine/Plugins/Developer/Concert/ConcertSync/ConcertSyncTest/Source/ConcertSyncTest/Private/Replication/Util/SendReceiveObjectTestBase.cpp
--- a/Engine/Plugins/Developer/Concert/ConcertSync/ConcertSyncTest/Source/ConcertSyncTest/Private/Replication/Util/SendReceiveObjectTestBase.cpp
+++ b/Engine/Plugins/Developer/Concert/ConcertSync/ConcertSyncTest/Source/ConcertSyncTest/Private/Replication/Util/SendReceiveObjectTestBase.cpp
@@ -47,7 +47,9 @@ namespace UE::ConcertSyncTests::Replication
 	
 	void FSendReceiveObjectTestBase::SimulateSendObjectToReceiver(
 		TFunctionRef<FReceiveReplicationEventSignature> OnServerReceive,
-		TFunctionRef<FReceiveReplicationEventSignature> OnReceiverClientReceive)
+		TFunctionRef<FReceiveReplicationEventSignature> OnReceiverClientReceive,
+		EPropertyTestFlags PropertyFlags
+		)
 	{
 		auto TestReplicationData_Server = [this, OnServerReceive](const FConcertSessionContext& Context, const FConcertBatchReplicationEvent& Event)
 		{
@@ -82,37 +84,67 @@ namespace UE::ConcertSyncTests::Replication
 		// TestObject is the same UObject on both clients.
 		// Hence we must override test values with SetTestValues and SetDifferentValues.
 		// 1 Sender > Server
-		SetTestValues(*TestObject);
+		SetTestValues(*TestObject, PropertyFlags);
 		TickClient(Client_Sender);
 		
 		// 2 Forward from server to receiver
 		TickServer();
 		
 		// 3 Receive from server
-		SetDifferentValues(*TestObject);
+		SetDifferentValues(*TestObject, PropertyFlags);
 		TickClient(Client_Receiver);
 
 		ServerSession->UnregisterCustomEventHandler<FConcertBatchReplicationEvent>(ServerHandle);
-		ServerSession->UnregisterCustomEventHandler<FConcertBatchReplicationEvent>(ClientHandle);
+		Client_Receiver->ClientSessionMock->UnregisterCustomEventHandler<FConcertBatchReplicationEvent>(ClientHandle);
 		
 		// No call to Super because we're completely overriding the behavior.
 	}
 
-	void FSendReceiveObjectTestBase::SetTestValues(UTestReflectionObject& Object)
+	void FSendReceiveObjectTestBase::SetTestValues(UTestReflectionObject& Object, EPropertyTestFlags PropertyFlags)
 	{
-		Object.Float = SentFloat;
-		Object.Vector = SentVector;
+		if (EnumHasAnyFlags(PropertyFlags, EPropertyTestFlags::Float))
+		{
+			Object.Float = SentFloat;
+		}
+		if (EnumHasAnyFlags(PropertyFlags, EPropertyTestFlags::Vector))
+		{
+			Object.Vector = SentVector;
+		}
 	}
 	
-	void FSendReceiveObjectTestBase::SetDifferentValues(UTestReflectionObject& Object)
+	void FSendReceiveObjectTestBase::SetDifferentValues(UTestReflectionObject& Object, EPropertyTestFlags PropertyFlags)
 	{
-		Object.Float *= -1.f;
-		Object.Vector *= -1.f;
+		if (EnumHasAnyFlags(PropertyFlags, EPropertyTestFlags::Float))
+		{
+			Object.Float = DifferentFloat;
+		}
+		if (EnumHasAnyFlags(PropertyFlags, EPropertyTestFlags::Vector))
+		{
+			Object.Vector = DifferentVector;
+		}
 	}
 	
-	void FSendReceiveObjectTestBase::TestEqualTestValues(UTestReflectionObject& Object, FAutomationTestBase& Test)
+	void FSendReceiveObjectTestBase::TestEqualTestValues(UTestReflectionObject& Object, EPropertyTestFlags PropertyFlags)
 	{
-		Test.TestEqual(TEXT("Float"), Object.Float, SentFloat);
-		Test.TestEqual(TEXT("Vector"), Object.Vector, SentVector);
+		if (EnumHasAnyFlags(PropertyFlags, EPropertyTestFlags::Float))
+		{
+			TestEqual(TEXT("Float"), Object.Float, SentFloat);
+		}
+		if (EnumHasAnyFlags(PropertyFlags, EPropertyTestFlags::Vector))
+		{
+			TestEqual(TEXT("Vector"), Object.Vector, SentVector);
+		}
+	}
+
+	void FSendReceiveObjectTestBase::TestEqualDifferentValues(UTestReflectionObject& Object, EPropertyTestFlags PropertyFlags)
+	{
+		if (EnumHasAnyFlags(PropertyFlags, EPropertyTestFlags::Float))
+		{
+			TestEqual(TEXT("Float"), Object.Float, DifferentFloat);
+		}
+		if (EnumHasAnyFlags(PropertyFlags, EPropertyTestFlags::Vector))
+		{
+			TestEqual(TEXT("Vector"), Object.Vector, DifferentVector);
+		}
 	}
-}
\ No newline at end of file
+}
diff --git a/Engine/Plugins/Developer/Concert/ConcertSync/ConcertSyncTest/Source/ConcertSyncTest/Private/Replication/Util/SendReceiveObjectTestBase.h b/Engine/Plugins/Developer/Concert/ConcertSync/ConcertSyncTest/Source/ConcertSyncTest/Private/Replication/Util/SendReceiveObjectTestBase.h
--- a/Engine/Plugins/Developer/Concert/ConcertSync/ConcertSyncTest/Source/ConcertSyncTest/Private/Replication/Util/SendReceiveObjectTestBase.h
+++ b/Engine/Plugins/Developer/Concert/ConcertSync/ConcertSyncTest/Source/ConcertSyncTest/Private/Replication/Util/SendReceiveObjectTestBase.h
@@ -16,6 +16,18 @@ namespace UE::ConcertSyncServer::Replication
 namespace UE::ConcertSyncTests::Replication
 {
 	class FConcertClientReplicationBridgeMock;
+
+	/** Defines the properties to be tested in FSendReceiveObjectTestBase: some tests may want more control of which properties are sent. */
+	enum class EPropertyTestFlags
+	{
+		None,
+		
+		Float = 1 << 0,
+		Vector = 1 << 1,
+		
+		All = Float | Vector
+	};
+	ENUM_CLASS_FLAGS(EPropertyTestFlags);
 	
 	/**
 	 * Creates a server and connects a sender and receiver client.
@@ -36,6 +48,7 @@ namespace UE::ConcertSyncTests::Replication
 		 * there really is not need to override CleanUpTest because we assign TestObject first thing in SetUpClientAndServer.
 		 */
 		UTestReflectionObject* TestObject = nullptr;
+		const FGuid SenderStreamId = FGuid::NewGuid();
 
 		//~ Begin FSendReceiveTestBase Interface
 		virtual ConcertSyncClient::Replication::FJoinReplicatedSessionArgs CreateSenderArgs() override;
@@ -50,20 +63,28 @@ namespace UE::ConcertSyncTests::Replication
 		 */
 		void SimulateSendObjectToReceiver(
 			TFunctionRef<FReceiveReplicationEventSignature> OnServerReceive = [](auto&, auto&){},
-			TFunctionRef<FReceiveReplicationEventSignature> OnReceiverClientReceive = [](auto&, auto&){}
+			TFunctionRef<FReceiveReplicationEventSignature> OnReceiverClientReceive = [](auto&, auto&){},
+			EPropertyTestFlags PropertyFlags = EPropertyTestFlags::All
 			);
 
-		void SetTestValues(UTestReflectionObject& Object);
-		void SetDifferentValues(UTestReflectionObject& Object);
-		void TestEqualTestValues(UTestReflectionObject& Object, FAutomationTestBase& Test);
+		/** Sets the specified properties to its test values */
+		void SetTestValues(UTestReflectionObject& Object, EPropertyTestFlags PropertyFlags = EPropertyTestFlags::All);
+		/** Sets the specified properties to values different from the test values */
+		void SetDifferentValues(UTestReflectionObject& Object, EPropertyTestFlags PropertyFlags = EPropertyTestFlags::All);
+
+		/** Tests that the specified properties are equal to their test values */
+		void TestEqualTestValues(UTestReflectionObject& Object, EPropertyTestFlags PropertyFlags = EPropertyTestFlags::All);
+		/** Tests that the specified properties are equal to the values different from the test values (i.e. the values SetDifferentValues sets). */
+		void TestEqualDifferentValues(UTestReflectionObject& Object, EPropertyTestFlags PropertyFlags = EPropertyTestFlags::All);
 		
 	protected:
 		
 		const float SentFloat = 42.f;
 		const FVector SentVector = { 21.f, 84.f, -1.f };
 		
-		const FGuid SenderStreamId = FGuid::NewGuid();
-
+		const float DifferentFloat = -420.f;
+		const FVector DifferentVector = { -210.f, -840.f, 10.f };
+		
 		virtual TSet<FGuid> GetSenderStreamIds() const { return { SenderStreamId }; } 
 	};
 }
diff --git a/Engine/Plugins/Developer/Concert/ConcertSync/ConcertSyncTest/Source/ConcertSyncTest/Private/Replication/Util/SendReceiveTestBase.cpp b/Engine/Plugins/Developer/Concert/ConcertSync/ConcertSyncTest/Source/ConcertSyncTest/Private/Replication/Util/SendReceiveTestBase.cpp
--- a/Engine/Plugins/Developer/Concert/ConcertSync/ConcertSyncTest/Source/ConcertSyncTest/Private/Replication/Util/SendReceiveTestBase.cpp
+++ b/Engine/Plugins/Developer/Concert/ConcertSync/ConcertSyncTest/Source/ConcertSyncTest/Private/Replication/Util/SendReceiveTestBase.cpp
@@ -105,7 +105,7 @@ namespace UE::ConcertSyncTests::Replication
 		TickClient(Client_Receiver);
 		
 		ServerSession->UnregisterCustomEventHandler<FConcertBatchReplicationEvent>(ServerHandle);
-		ServerSession->UnregisterCustomEventHandler<FConcertBatchReplicationEvent>(ClientHandle);
+		Client_Receiver->ClientSessionMock->UnregisterCustomEventHandler<FConcertBatchReplicationEvent>(ClientHandle);
 	}
 	
 	void FSendReceiveTestBase::TickClient(FClientInfo* Client)
