Extract authority conflict detection to ConcertSyncCore

diff --git a/Engine/Plugins/Developer/Concert/ConcertSync/ConcertSyncCore/Source/ConcertSyncCore/Private/Replication/AuthorityConflictSharedUtils.cpp b/Engine/Plugins/Developer/Concert/ConcertSync/ConcertSyncCore/Source/ConcertSyncCore/Private/Replication/AuthorityConflictSharedUtils.cpp
new file mode 100644
--- /dev/null
+++ b/Engine/Plugins/Developer/Concert/ConcertSync/ConcertSyncCore/Source/ConcertSyncCore/Private/Replication/AuthorityConflictSharedUtils.cpp
@@ -0,0 +1,86 @@
+// Copyright Epic Games, Inc. All Rights Reserved.
+
+#include "Replication/AuthorityConflictSharedUtils.h"
+
+#include "Replication/Data/ObjectIds.h"
+#include "Replication/Data/ReplicationStreamDescription.h"
+
+namespace UE::ConcertSyncCore::Replication::AuthorityConflictUtils
+{
+	namespace Private
+	{
+		static void ForEachClientWithPotentialConflict(
+            const FSoftObjectPath& Object,
+            TFunctionRef<EBreakBehavior(const FGuid& ClientId, const FGuid& StreamId, const FConcertPropertySelection& ConflictingProperty)> Callback,
+            TArrayView<const FGuid> IgnoredClients,
+            const IReplicationGroundTruth& GroundTruth
+            )
+		{
+			GroundTruth.ForEachSendingClient([&Object, &Callback, &IgnoredClients, &GroundTruth](const FGuid& ClientEndpointId)
+			{
+				if (IgnoredClients.Contains(ClientEndpointId))
+				{
+					return EBreakBehavior::Continue;
+				}
+				
+				EBreakBehavior Result = EBreakBehavior::Continue;
+				GroundTruth.ForEachStream(ClientEndpointId, [&Object, &Callback, &ClientEndpointId, &GroundTruth, &Result](const FSharedReplicationStreamDescription& Stream) mutable
+				{
+					// If client has not claimed authority over this object in this stream, skip
+					const FGuid& StreamId = Stream.Identifier;
+					if (!GroundTruth.HasAuthority(ClientEndpointId, StreamId, Object))
+					{
+						return EBreakBehavior::Continue;
+					}
+
+					// This client is using the request object: report the potential conflict ...
+					const FObjectReplicationMap& ObjectReplicationMap = Stream.ReplicationMap;
+					const FReplicatedObjectInfo* ReplicationObjectInfo = ObjectReplicationMap.ReplicatedObjects.Find(Object);
+					if (ReplicationObjectInfo && Callback(ClientEndpointId, StreamId, ReplicationObjectInfo->PropertySelection) == EBreakBehavior::Break)
+					{
+						// ... conflict ends iteration
+						Result = EBreakBehavior::Break;
+						return EBreakBehavior::Break;
+					}
+					// ... conflict resolved
+					return EBreakBehavior::Continue;
+				});
+				return Result;
+			});
+		}
+	}
+	
+	EAuthorityConflict EnumerateAuthorityConflicts(
+		const FGuid& ClientId,
+		const FSoftObjectPath& Object,
+		TConstArrayView<FConcertPropertyChain> OverwriteProperties,
+		const IReplicationGroundTruth& GroundTruth,
+		FProcessAuthorityConflict ProcessConflict
+		)
+	{
+		const FGuid IgnoredClients[] = { ClientId };
+
+		bool bFreeOfConflicts = true;
+		Private::ForEachClientWithPotentialConflict(
+			Object,
+			[&ProcessConflict, &OverwriteProperties, &bFreeOfConflicts](const FGuid& ClientId, const FGuid& StreamId, const FConcertPropertySelection& WrittenProperties) mutable
+			{
+				EBreakBehavior BreakBehavior = EBreakBehavior::Continue;
+				const bool bHasOverlap = FConcertPropertySelection::EnumeratePropertyOverlaps(WrittenProperties.ReplicatedProperties, OverwriteProperties,
+					[&ProcessConflict, &ClientId, &StreamId, &BreakBehavior](const FConcertPropertyChain& Overlap)
+					{
+						BreakBehavior = ProcessConflict(ClientId, StreamId, Overlap);
+						return BreakBehavior;
+					});
+				
+				// At this point we know that WritingClientId is sending WrittenProperties - if the properties overlap, the authority request is not possible
+				const bool bHasNoConflict = !bHasOverlap;
+				bFreeOfConflicts &= bHasNoConflict;
+				return BreakBehavior;
+			},
+			IgnoredClients,
+			GroundTruth
+			);
+		return bFreeOfConflicts ? EAuthorityConflict::Allowed : EAuthorityConflict::Conflict;
+	}
+}
diff --git a/Engine/Plugins/Developer/Concert/ConcertSync/ConcertSyncCore/Source/ConcertSyncCore/Private/Replication/Data/ConcertPropertySelection.cpp b/Engine/Plugins/Developer/Concert/ConcertSync/ConcertSyncCore/Source/ConcertSyncCore/Private/Replication/Data/ConcertPropertySelection.cpp
--- a/Engine/Plugins/Developer/Concert/ConcertSync/ConcertSyncCore/Source/ConcertSyncCore/Private/Replication/Data/ConcertPropertySelection.cpp
+++ b/Engine/Plugins/Developer/Concert/ConcertSync/ConcertSyncCore/Source/ConcertSyncCore/Private/Replication/Data/ConcertPropertySelection.cpp
@@ -199,26 +199,33 @@ FString FConcertPropertyChain::ToString(EToStringMethod Method) const
 		checkNoEntry();
 		return FName(NAME_None).ToString();
 	}
-	
 }
 
-bool FConcertPropertySelection::OverlapsWith(const FConcertPropertySelection& Other) const
+bool FConcertPropertySelection::EnumeratePropertyOverlaps(
+	TConstArrayView<FConcertPropertyChain> First,
+	TConstArrayView<FConcertPropertyChain> Second,
+	TFunctionRef<EBreakBehavior(const FConcertPropertyChain&)> Callback
+	)
 {
 	/* 
 	 * Performance could be improved: Two chains can only overlap if their root properties overlap.
 	 * Example: [Foo.Vector.X] and [Foo.FloatProperty] both share Foo struct (but do not overlap in this case).
 	 */
-	for (const FConcertPropertyChain& ThisChain : ReplicatedProperties)
+	bool bHasOverlap = false;
+	for (const FConcertPropertyChain& ThisChain : First)
 	{
-		for (const FConcertPropertyChain& OtherChain : Other.ReplicatedProperties)
+		for (const FConcertPropertyChain& OtherChain : Second)
 		{
-			if (ThisChain == OtherChain) 
+			const bool bIsOverlap = ThisChain == OtherChain;
+			bHasOverlap |= bIsOverlap;
+			if (bIsOverlap && Callback(ThisChain) == EBreakBehavior::Break) 
 			{
 				return true;
 			}
 		}
 	}
-	return false;
+
+	return bHasOverlap;
 }
 
 uint32 GetTypeHash(const FConcertPropertyChain& Chain)
diff --git a/Engine/Plugins/Developer/Concert/ConcertSync/ConcertSyncCore/Source/ConcertSyncCore/Public/Replication/AuthorityConflictSharedUtils.h b/Engine/Plugins/Developer/Concert/ConcertSync/ConcertSyncCore/Source/ConcertSyncCore/Public/Replication/AuthorityConflictSharedUtils.h
new file mode 100644
--- /dev/null
+++ b/Engine/Plugins/Developer/Concert/ConcertSync/ConcertSyncCore/Source/ConcertSyncCore/Public/Replication/AuthorityConflictSharedUtils.h
@@ -0,0 +1,63 @@
+// Copyright Epic Games, Inc. All Rights Reserved.
+
+#pragma once
+
+#include "Containers/Array.h"
+#include "HAL/Platform.h"
+#include "Misc/EBreakBehavior.h"
+#include "Templates/Function.h"
+
+struct FConcertPropertyChain;
+struct FConcertPropertySelection;
+struct FReplicatedObjectId;
+struct FSharedReplicationStreamDescription;
+
+namespace UE::ConcertSyncCore::Replication::AuthorityConflictUtils
+{
+	enum class EAuthorityConflict : uint8
+	{
+		/** The client is allowed to take authority */
+		Allowed,
+		/** There was at least one conflict */
+		Conflict,
+	};
+	
+	using FProcessAuthorityConflict = TFunctionRef<EBreakBehavior(const FGuid& ClientId, const FGuid& StreamId, const FConcertPropertyChain& ConflictingProperty)>;
+
+	/** Defines a way to obtain streaming information about clients. */
+	class IReplicationGroundTruth
+	{
+	public:
+
+		/** Provides a way to extract all streams registered to a given client. */
+		virtual void ForEachStream(const FGuid& ClientEndpointId, TFunctionRef<EBreakBehavior(const FSharedReplicationStreamDescription& Stream)> Callback) const = 0;
+
+		/** Iterates through all clients have registered to send any data. */
+		virtual void ForEachSendingClient(TFunctionRef<EBreakBehavior(const FGuid& ClientEndpointId)> Callback) const = 0;
+
+		/** @return Whether ClientId's stream StreamId has authority over ObjectPath. */
+		virtual bool HasAuthority(const FGuid& ClientId, const FGuid& StreamId, const FSoftObjectPath& ObjectPath) const = 0; 
+
+		virtual ~IReplicationGroundTruth() = default;
+	};
+	
+	/**
+	 * Enumerates all authority conflicts, if any, that would occur if Object were to be replicated with the given Properties.
+	 *
+	 * @param ClientId The client that wants to take authority 
+	 * @param Object Describes the object, the client, and the stream to be replicated.
+	 * @param OverwriteProperties The properties that will be registered with the object (i.e. not what IClientInfoSource reports currently)
+	 * @param GroundTruth Provides information about clients. On the server, it is the definite state whereas on client machines it would be
+	 * what the local client thinks what the server state is.
+	 * @param ProcessConflict Callback to be called for each permutation of conflicting client, stream and property.
+	 * 
+	 * @return Whether there were any conflicts
+	 */
+	CONCERTSYNCCORE_API EAuthorityConflict EnumerateAuthorityConflicts(
+		const FGuid& ClientId,
+		const FSoftObjectPath& Object,
+		TConstArrayView<FConcertPropertyChain> OverwriteProperties,
+		const IReplicationGroundTruth& GroundTruth,
+		FProcessAuthorityConflict ProcessConflict = [](const FGuid&, const FGuid&, const FConcertPropertyChain&){ return EBreakBehavior::Break; }
+		);
+}
diff --git a/Engine/Plugins/Developer/Concert/ConcertSync/ConcertSyncCore/Source/ConcertSyncCore/Public/Replication/Data/ConcertPropertySelection.h b/Engine/Plugins/Developer/Concert/ConcertSync/ConcertSyncCore/Source/ConcertSyncCore/Public/Replication/Data/ConcertPropertySelection.h
--- a/Engine/Plugins/Developer/Concert/ConcertSync/ConcertSyncCore/Source/ConcertSyncCore/Public/Replication/Data/ConcertPropertySelection.h
+++ b/Engine/Plugins/Developer/Concert/ConcertSync/ConcertSyncCore/Source/ConcertSyncCore/Public/Replication/Data/ConcertPropertySelection.h
@@ -2,6 +2,7 @@
 
 #pragma once
 
+#include "Misc/EBreakBehavior.h"
 #include "ConcertPropertySelection.generated.h"
 
 class UStruct;
@@ -181,8 +182,20 @@ struct CONCERTSYNCCORE_API FConcertPropertySelection
 	UPROPERTY()
 	TArray<FConcertPropertyChain> ReplicatedProperties;
 
-	/** @return Whether this and Other contain at least one property that is the same. This algorithm is strictly O(n^2) but runs O(n) on average. */
-	bool OverlapsWith(const FConcertPropertySelection& Other) const;
+	/** @return Whether this and Other contain at least one property that is the same. */
+	bool OverlapsWith(const FConcertPropertySelection& Other) const { return EnumeratePropertyOverlaps(ReplicatedProperties, Other.ReplicatedProperties); }
+
+	/**
+	 * Determines all properties that overlap.
+	 * This algorithm is strictly O(n^2) but runs O(n) on average.
+	 * 
+	 * @return Whether there were any property overlaps.
+	 */
+	static bool EnumeratePropertyOverlaps(
+		TConstArrayView<FConcertPropertyChain> First,
+		TConstArrayView<FConcertPropertyChain> Second,
+		TFunctionRef<EBreakBehavior(const FConcertPropertyChain&)> Callback = [](const FConcertPropertyChain&){ return EBreakBehavior::Break; }
+		);
 
 	friend bool operator==(const FConcertPropertySelection& Left, const FConcertPropertySelection& Right)
 	{
diff --git a/Engine/Plugins/Developer/Concert/ConcertSync/ConcertSyncServer/Source/ConcertSyncServer/Private/Replication/AuthorityManager.cpp b/Engine/Plugins/Developer/Concert/ConcertSync/ConcertSyncServer/Source/ConcertSyncServer/Private/Replication/AuthorityManager.cpp
--- a/Engine/Plugins/Developer/Concert/ConcertSync/ConcertSyncServer/Source/ConcertSyncServer/Private/Replication/AuthorityManager.cpp
+++ b/Engine/Plugins/Developer/Concert/ConcertSync/ConcertSyncServer/Source/ConcertSyncServer/Private/Replication/AuthorityManager.cpp
@@ -4,6 +4,7 @@
 
 #include "ConcertLogGlobal.h"
 #include "IConcertSession.h"
+#include "Replication/AuthorityConflictSharedUtils.h"
 #include "Replication/Data/ObjectIds.h"
 #include "Replication/Data/ReplicationStreamDescription.h"
 #include "Replication/Messages/ChangeAuthority.h"
@@ -12,6 +13,33 @@ namespace UE::ConcertSyncServer::Replication
 {
 	namespace Private
 	{
+		class FServerGroundTruth : public ConcertSyncCore::Replication::AuthorityConflictUtils::IReplicationGroundTruth
+		{
+			const FAuthorityManager& Owner;
+			const IAuthorityManagerGetters& Getters;
+		public:
+			
+			FServerGroundTruth(const FAuthorityManager& Owner, const IAuthorityManagerGetters& Getters)
+				: Owner(Owner)
+				, Getters(Getters)
+			{}
+
+			virtual void ForEachStream(const FGuid& ClientEndpointId, TFunctionRef<EBreakBehavior(const FSharedReplicationStreamDescription& Stream)> Callback) const override
+			{
+				Getters.ForEachStream(ClientEndpointId, [&Callback](const FReplicationStreamDescription& Stream){ return Callback(Stream.BaseDescription); });
+			}
+			
+			virtual void ForEachSendingClient(TFunctionRef<EBreakBehavior(const FGuid& ClientEndpointId)> Callback) const override
+			{
+				Getters.ForEachSendingClient(Callback);
+			}
+			
+			virtual bool HasAuthority(const FGuid& ClientId, const FGuid& StreamId, const FSoftObjectPath& ObjectPath) const override
+			{
+				return Owner.HasAuthorityToChange({ { StreamId, ObjectPath }, ClientId });
+			}
+		};
+		
 		static void ForEachReplicatedObject(
 			const TMap<FSoftObjectPath, FConcertStreamArray>& Map,
 			TFunctionRef<void(const FGuid& StreamId, const FSoftObjectPath& ObjectPath)> Callback
@@ -84,22 +112,19 @@ namespace UE::ConcertSyncServer::Replication
 			return EAuthorityResult::NoRegisteredProperties;
 		}
 
-		bool bFreeOfPropertyOverlaps = true;
-		const FClientId IgnoredClients[] = { ClientId };
-		ForEachClientWithPotentialConflict(
+		using namespace ConcertSyncCore::Replication;
+		using namespace ConcertSyncCore::Replication::AuthorityConflictUtils;
+		Private::FServerGroundTruth GroundTruth(*this, Getters);
+		const EAuthorityConflict Conflict = AuthorityConflictUtils::EnumerateAuthorityConflicts(
+			Object.SenderEndpointId,
 			Object.Object,
-			[this, &ProcessConflict, PropertiesToCheck, &bFreeOfPropertyOverlaps](const FClientId& ClientId, const FStreamId& StreamId, const FConcertPropertySelection& WrittenProperties) mutable
+			PropertiesToCheck->ReplicatedProperties,
+			GroundTruth,
+			[&ProcessConflict](const FClientId& ClientId, const FStreamId& StreamId, const FConcertPropertyChain& Property)
 			{
-				// At this point we know that WritingClientId is sending WrittenProperties - if the properties overlap, the authority request is not possible
-				const bool bHasNoConflict = !WrittenProperties.OverlapsWith(*PropertiesToCheck);
-				bFreeOfPropertyOverlaps &= bHasNoConflict;
-				return bHasNoConflict
-					? EBreakBehavior::Continue
-					: ProcessConflict(ClientId, StreamId, WrittenProperties);
-			},
-			IgnoredClients
-			);
-		return bFreeOfPropertyOverlaps ? EAuthorityResult::Allowed : EAuthorityResult::Conflict;
+				return ProcessConflict(ClientId, StreamId, Property);
+			});
+		return Conflict == EAuthorityConflict::Allowed ? EAuthorityResult::Allowed : EAuthorityResult::Conflict;
 	}
 
 	bool FAuthorityManager::CanTakeAuthority(const FReplicatedObjectId& Object) const
@@ -203,52 +228,4 @@ namespace UE::ConcertSyncServer::Replication
 		});
 		return StreamDescription;
 	}
-
-	void FAuthorityManager::ForEachClientWithPotentialConflict(
-		const FSoftObjectPath& Object,
-		FProcessAuthorityConflict Callback,
-		TArrayView<const FClientId> IgnoredClients
-		) const
-	{
-		Getters.ForEachSendingClient([this, &Object, &Callback, &IgnoredClients](const FGuid& ClientEndpointId)
-		{
-			if (IgnoredClients.Contains(ClientEndpointId))
-			{
-				return EBreakBehavior::Continue;
-			}
-
-			// If client has no authority at all, skip.
-			const FClientAuthorityData* OtherClientData = ClientAuthorityData.Find(ClientEndpointId);
-			if (!OtherClientData)
-			{
-				return EBreakBehavior::Continue;
-			}
-			
-			EBreakBehavior Result = EBreakBehavior::Continue;
-			Getters.ForEachStream(ClientEndpointId, [&Object, &Callback, &ClientEndpointId, OtherClientData, &Result](const FReplicationStreamDescription& Stream) mutable
-			{
-				// If client has not claimed authority over this object in this stream, skip
-				const FStreamId& StreamId = Stream.BaseDescription.Identifier;
-				const TSet<FSoftObjectPath>* ClientControlledObjects = OtherClientData->OwnedObjects.Find(StreamId);
-				const bool bClientHasSomeAuthorityOverObject = ClientControlledObjects && ClientControlledObjects->Contains(Object);
-				if (!bClientHasSomeAuthorityOverObject)
-				{
-					return EBreakBehavior::Continue;
-				}
-
-				// This client is using the request object: report the potential conflict ...
-				const FObjectReplicationMap& ObjectReplicationMap = Stream.BaseDescription.ReplicationMap;
-				const FReplicatedObjectInfo* ReplicationObjectInfo = ObjectReplicationMap.ReplicatedObjects.Find(Object);
-				if (ReplicationObjectInfo && Callback(ClientEndpointId, StreamId, ReplicationObjectInfo->PropertySelection) == EBreakBehavior::Break)
-				{
-					// ... conflict ends iteration
-					Result = EBreakBehavior::Break;
-					return EBreakBehavior::Break;
-				}
-				// ... conflict resolved
-				return EBreakBehavior::Continue;
-			});
-			return Result;
-		});
-	}
 }
diff --git a/Engine/Plugins/Developer/Concert/ConcertSync/ConcertSyncServer/Source/ConcertSyncServer/Private/Replication/AuthorityManager.h b/Engine/Plugins/Developer/Concert/ConcertSync/ConcertSyncServer/Source/ConcertSyncServer/Private/Replication/AuthorityManager.h
--- a/Engine/Plugins/Developer/Concert/ConcertSync/ConcertSyncServer/Source/ConcertSyncServer/Private/Replication/AuthorityManager.h
+++ b/Engine/Plugins/Developer/Concert/ConcertSync/ConcertSyncServer/Source/ConcertSyncServer/Private/Replication/AuthorityManager.h
@@ -11,6 +11,7 @@ class IConcertSession;
 
 struct FConcertReplication_ChangeAuthority_Response;
 struct FConcertReplication_ChangeAuthority_Request;
+struct FConcertPropertyChain;
 struct FConcertPropertySelection;
 struct FConcertSessionContext;
 struct FObjectInStreamID;
@@ -28,10 +29,10 @@ namespace UE::ConcertSyncServer::Replication
 	public:
 
 		/** Provides a way to extract all streams registered to a given client. */
-		virtual void ForEachStream(const FGuid& ClientEndpointId, TFunctionRef<EBreakBehavior(const FReplicationStreamDescription& Stream)> Callback) = 0;
+		virtual void ForEachStream(const FGuid& ClientEndpointId, TFunctionRef<EBreakBehavior(const FReplicationStreamDescription& Stream)> Callback) const = 0;
 
 		/** Iterates through all clients have registered to send any data. */
-		virtual void ForEachSendingClient(TFunctionRef<EBreakBehavior(const FGuid& ClientEndpointId)> Callback) = 0;
+		virtual void ForEachSendingClient(TFunctionRef<EBreakBehavior(const FGuid& ClientEndpointId)> Callback) const = 0;
 
 		virtual ~IAuthorityManagerGetters() = default;
 	};
@@ -43,7 +44,7 @@ namespace UE::ConcertSyncServer::Replication
 		
 		using FStreamId = FGuid;
 		using FClientId = FGuid;
-		using FProcessAuthorityConflict = TFunctionRef<EBreakBehavior(const FClientId& ClientId, const FStreamId& StreamId, const FConcertPropertySelection& WrittenProperties)>;
+		using FProcessAuthorityConflict = TFunctionRef<EBreakBehavior(const FClientId& ClientId, const FStreamId& StreamId, const FConcertPropertyChain& WrittenProperties)>;
 
 		FAuthorityManager(IAuthorityManagerGetters& Getters, TSharedRef<IConcertSession> InSession);
 		~FAuthorityManager();
@@ -106,13 +107,6 @@ namespace UE::ConcertSyncServer::Replication
 
 		/** Finds a stream registered with the client by its ID. */
 		const FReplicationStreamDescription* FindClientStreamById(const FClientId& ClientId, const FStreamId& StreamId) const;
-
-		/** Iterates through all clients that are already writing to Object. */
-		void ForEachClientWithPotentialConflict(
-			const FSoftObjectPath& Object,
-			FProcessAuthorityConflict Callback,
-			TArrayView<const FClientId> IgnoredClients = {}
-			) const;
 	};
 }
 
diff --git a/Engine/Plugins/Developer/Concert/ConcertSync/ConcertSyncServer/Source/ConcertSyncServer/Private/Replication/ConcertServerReplicationManager.cpp b/Engine/Plugins/Developer/Concert/ConcertSync/ConcertSyncServer/Source/ConcertSyncServer/Private/Replication/ConcertServerReplicationManager.cpp
--- a/Engine/Plugins/Developer/Concert/ConcertSync/ConcertSyncServer/Source/ConcertSyncServer/Private/Replication/ConcertServerReplicationManager.cpp
+++ b/Engine/Plugins/Developer/Concert/ConcertSync/ConcertSyncServer/Source/ConcertSyncServer/Private/Replication/ConcertServerReplicationManager.cpp
@@ -41,7 +41,7 @@ namespace UE::ConcertSyncServer::Replication
 		Session->OnTick().RemoveAll(this);
 	}
 
-	void FConcertServerReplicationManager::ForEachStream(const FGuid& ClientEndpointId, TFunctionRef<EBreakBehavior(const FReplicationStreamDescription& Stream)> Callback)
+	void FConcertServerReplicationManager::ForEachStream(const FGuid& ClientEndpointId, TFunctionRef<EBreakBehavior(const FReplicationStreamDescription& Stream)> Callback) const
 	{
 		const TSharedRef<FConcertReplicationClient>* Client = Clients.Find(ClientEndpointId);
 		if (!ensure(Client))
@@ -58,7 +58,7 @@ namespace UE::ConcertSyncServer::Replication
 		}
 	}
 
-	void FConcertServerReplicationManager::ForEachSendingClient(TFunctionRef<EBreakBehavior(const FGuid& ClientEndpointId)> Callback)
+	void FConcertServerReplicationManager::ForEachSendingClient(TFunctionRef<EBreakBehavior(const FGuid& ClientEndpointId)> Callback) const
 	{
 		for (const TPair<FGuid, TSharedRef<FConcertReplicationClient>>& ClientPair : Clients)
 		{
diff --git a/Engine/Plugins/Developer/Concert/ConcertSync/ConcertSyncServer/Source/ConcertSyncServer/Private/Replication/ConcertServerReplicationManager.h b/Engine/Plugins/Developer/Concert/ConcertSync/ConcertSyncServer/Source/ConcertSyncServer/Private/Replication/ConcertServerReplicationManager.h
--- a/Engine/Plugins/Developer/Concert/ConcertSync/ConcertSyncServer/Source/ConcertSyncServer/Private/Replication/ConcertServerReplicationManager.h
+++ b/Engine/Plugins/Developer/Concert/ConcertSync/ConcertSyncServer/Source/ConcertSyncServer/Private/Replication/ConcertServerReplicationManager.h
@@ -50,8 +50,8 @@ namespace UE::ConcertSyncServer::Replication
 		const FAuthorityManager& GetAuthorityManager() const { return AuthorityManager.Get(); }
 
 		//~ Begin IAuthorityManagerGetters Interface
-		virtual void ForEachStream(const FGuid& ClientEndpointId, TFunctionRef<EBreakBehavior(const FReplicationStreamDescription& Stream)> Callback) override;
-		virtual void ForEachSendingClient(TFunctionRef<EBreakBehavior(const FGuid& ClientEndpointId)> Callback) override;
+		virtual void ForEachStream(const FGuid& ClientEndpointId, TFunctionRef<EBreakBehavior(const FReplicationStreamDescription& Stream)> Callback) const override;
+		virtual void ForEachSendingClient(TFunctionRef<EBreakBehavior(const FGuid& ClientEndpointId)> Callback) const override;
 		//~ End IAuthorityManagerGetters Interface
 
 	private:
diff --git a/Engine/Plugins/Developer/Concert/ConcertSync/ConcertSyncServer/Source/ConcertSyncServer/Private/Replication/ConcertServerReplicationManager_ChangeStream.cpp b/Engine/Plugins/Developer/Concert/ConcertSync/ConcertSyncServer/Source/ConcertSyncServer/Private/Replication/ConcertServerReplicationManager_ChangeStream.cpp
--- a/Engine/Plugins/Developer/Concert/ConcertSync/ConcertSyncServer/Source/ConcertSyncServer/Private/Replication/ConcertServerReplicationManager_ChangeStream.cpp
+++ b/Engine/Plugins/Developer/Concert/ConcertSync/ConcertSyncServer/Source/ConcertSyncServer/Private/Replication/ConcertServerReplicationManager_ChangeStream.cpp
@@ -106,7 +106,7 @@ namespace UE::ConcertSyncServer::Replication
 
 				// Simply check whether any other client is already sending any of the requested properties.
 				AuthorityManager.EnumerateAuthorityConflicts(ReplicatedObjectInfo, &PropertySelection,
-					[&OutResponse, &ReplicatedObjectInfo](const FGuid& ClientId, const FGuid& StreamId, const FConcertPropertySelection& WrittenProperties)
+					[&OutResponse, &ReplicatedObjectInfo](const FGuid& ClientId, const FGuid& StreamId, const FConcertPropertyChain&)
 					{
 						const FReplicatedObjectId ConflictingObject = { { StreamId, ReplicatedObjectInfo.Object }, ClientId };
 						OutResponse.AuthorityConflicts.Add(ReplicatedObjectInfo, ConflictingObject);
diff --git a/Engine/Plugins/Developer/Concert/ConcertSync/ConcertSyncTest/Source/ConcertSyncTest/Private/Replication/ReplicationAuthorityTests.cpp b/Engine/Plugins/Developer/Concert/ConcertSync/ConcertSyncTest/Source/ConcertSyncTest/Private/Replication/ReplicationAuthorityTests.cpp
--- a/Engine/Plugins/Developer/Concert/ConcertSync/ConcertSyncTest/Source/ConcertSyncTest/Private/Replication/ReplicationAuthorityTests.cpp
+++ b/Engine/Plugins/Developer/Concert/ConcertSync/ConcertSyncTest/Source/ConcertSyncTest/Private/Replication/ReplicationAuthorityTests.cpp
@@ -13,6 +13,7 @@
 #include "Misc/AutomationTest.h"
 #include "Misc/Paths.h"
 #include "HAL/FileManager.h"
+#include "Replication/AuthorityConflictSharedUtils.h"
 
 namespace UE::ConcertSyncTests::Replication::Authority
 {
@@ -365,4 +366,151 @@ namespace UE::ConcertSyncTests::Replication::Authority
 		TestTrue(TEXT("Timed out ReleaseAuthorityOf reverted local server prediction"), ClientReplicationManager_Sender->GetClientOwnedObjects().Contains(TestObject));
 		return true;
 	}
+
+	namespace ConflictEnumerationComponent
+	{
+		constexpr FGuid RequestingClientId { 0, 0, 0, 0};
+		constexpr FGuid ExistingClientId { 1, 0, 0, 0};
+		
+		class FTestGroundTruth : public ConcertSyncCore::Replication::AuthorityConflictUtils::IReplicationGroundTruth
+		{
+		public:
+			
+			FSharedReplicationStreamDescription RequestingClientStream;
+			FSharedReplicationStreamDescription ExistingClientStream;
+			bool bRequestingClientHasAuthority;
+			bool bExistingClientHasAuthority;
+
+			FTestGroundTruth(const FSharedReplicationStreamDescription& RequestingClientStream, const FSharedReplicationStreamDescription& ExistingClientStream, bool bRequestingClientHasAuthority, bool bExistingClientHasAuthority)
+				: RequestingClientStream(RequestingClientStream)
+				, ExistingClientStream(ExistingClientStream)
+				, bRequestingClientHasAuthority(bRequestingClientHasAuthority)
+				, bExistingClientHasAuthority(bExistingClientHasAuthority)
+			{}
+
+			virtual void ForEachStream(const FGuid& ClientEndpointId, TFunctionRef<EBreakBehavior(const FSharedReplicationStreamDescription& Stream)> Callback) const override
+			{
+				if (ClientEndpointId == RequestingClientId)
+				{
+					Callback(RequestingClientStream);
+				}
+				if (ClientEndpointId == ExistingClientId)
+				{
+					Callback(ExistingClientStream);
+				}
+			}
+			
+			virtual void ForEachSendingClient(TFunctionRef<EBreakBehavior(const FGuid& ClientEndpointId)> Callback) const override
+			{
+				if (bRequestingClientHasAuthority && Callback(RequestingClientId) == EBreakBehavior::Break)
+				{
+					return;
+				}
+				if (bExistingClientHasAuthority)
+				{
+					Callback(ExistingClientId);
+				}
+			}
+			
+			virtual bool HasAuthority(const FGuid& ClientId, const FGuid& StreamId, const FSoftObjectPath& ObjectPath) const override
+			{
+				const bool bRequestor = ClientId == RequestingClientId && RequestingClientStream.Identifier == StreamId && bRequestingClientHasAuthority;
+				const bool bExisting = ClientId == ExistingClientId && ExistingClientStream.Identifier == StreamId && bExistingClientHasAuthority;
+				return bRequestor || bExisting;
+			}
+		};
+		
+		FSharedReplicationStreamDescription MakeStream(const FGuid& StreamId, const FSoftObjectPath ObjectPath, TArray<FConcertPropertyChain> Properties)
+		{
+			FObjectReplicationMap ExistingReplicationMap;
+			const FConcertPropertySelection Selection{ Properties};
+			ExistingReplicationMap.ReplicatedObjects.Add(ObjectPath, { {}, Selection });
+			const FSharedReplicationStreamDescription ExistingClientStream { StreamId, ExistingReplicationMap };
+			return ExistingClientStream;
+		};
+	}
+
+	/** Tests UE::ConcertSyncCore::Replication::AuthorityConflictUtils::EnumerateAuthorityConflicts. */
+	IMPLEMENT_CUSTOM_SIMPLE_AUTOMATION_TEST(FConflictEnumerationComponenTests, FSendReceiveObjectTestBase, "Concert.Replication.Authority.ConflictEnumerationComponenTests", EAutomationTestFlags::EditorContext | EAutomationTestFlags::EngineFilter);
+	bool FConflictEnumerationComponenTests::RunTest(const FString& Parameters)
+	{
+		using namespace ConflictEnumerationComponent;
+		using namespace ConcertSyncCore::Replication::AuthorityConflictUtils;
+		constexpr FGuid TestStreamId { 0, 0, 0, 1};
+		const FSoftObjectPath ObjectPath(TEXT("/Game/World.World:PersistentLevel.StaticMeshActor0"));
+		
+		FConcertPropertyChain Float = *FConcertPropertyChain::CreateFromPath(*UTestReflectionObject::StaticClass(), { TEXT("Float") });
+		FConcertPropertyChain Vector = *FConcertPropertyChain::CreateFromPath(*UTestReflectionObject::StaticClass(), { TEXT("Vector") });
+		FConcertPropertyChain VectorX = *FConcertPropertyChain::CreateFromPath(*UTestReflectionObject::StaticClass(), { TEXT("Vector"), TEXT("X") });
+		const FSharedReplicationStreamDescription ClientStream_Empty = MakeStream(TestStreamId, ObjectPath, {});
+		const FSharedReplicationStreamDescription ClientStream_FloatOnly = MakeStream(TestStreamId, ObjectPath, { Float });
+		const FSharedReplicationStreamDescription ClientStream_VectorOnly = MakeStream(TestStreamId, ObjectPath, { Vector, VectorX });
+
+		// Requestor tries to stream Float, then Vector
+		// Existing is replicating Float
+		{
+			const FTestGroundTruth GroundTruth { ClientStream_Empty, ClientStream_FloatOnly, false, true };
+
+			int32 NumCalls = 0;
+			const EAuthorityConflict FloatConflict = EnumerateAuthorityConflicts(RequestingClientId, ObjectPath, { Float }, GroundTruth,
+				[this, &TestStreamId, &Float, &NumCalls](const FGuid& ClientId, const FGuid& StreamId, const FConcertPropertyChain& ConflictingProperty)
+				{
+					++NumCalls;
+					TestEqual(TEXT("Float conflict > Client"), ClientId, ExistingClientId);
+					TestEqual(TEXT("Float conflict > Stream"), StreamId, TestStreamId);
+					TestEqual(TEXT("Float conflict > Property"), Float, ConflictingProperty);
+					return EBreakBehavior::Continue;
+				});
+			TestTrue(TEXT("1 Float overlaps"), FloatConflict == EAuthorityConflict::Conflict);
+			TestEqual(TEXT("1 Float > Exactly 1 conflicting property"), NumCalls, 1);
+			
+			const EAuthorityConflict VectorConflict = EnumerateAuthorityConflicts(RequestingClientId, ObjectPath, { Vector, VectorX }, GroundTruth,
+				[this](const FGuid& ClientId, const FGuid& StreamId, const FConcertPropertyChain& ConflictingProperty)
+				{
+					AddError(TEXT("No conflict expected"));
+					return EBreakBehavior::Continue;
+				});
+			TestTrue(TEXT("1 Vector does not overlap"), VectorConflict == EAuthorityConflict::Allowed);
+		}
+
+		// Requestor tries to stream Float, then Vector
+		// Existing is not replicating but has Float registered
+		{
+			const FTestGroundTruth GroundTruth { ClientStream_Empty, ClientStream_FloatOnly, false, false };
+			const EAuthorityConflict FloatConflict = EnumerateAuthorityConflicts(RequestingClientId, ObjectPath, { Float }, GroundTruth,
+				[this](const FGuid& ClientId, const FGuid& StreamId, const FConcertPropertyChain& ConflictingProperty)
+				{
+					AddError(TEXT("No conflict expected"));
+					return EBreakBehavior::Continue;
+				});
+			TestTrue(TEXT("2 Float does not overlap"), FloatConflict == EAuthorityConflict::Allowed);
+			const EAuthorityConflict VectorConflict = EnumerateAuthorityConflicts(RequestingClientId, ObjectPath, { Vector, VectorX }, GroundTruth,
+				[this](const FGuid& ClientId, const FGuid& StreamId, const FConcertPropertyChain& ConflictingProperty)
+				{
+					AddError(TEXT("No conflict expected"));
+					return EBreakBehavior::Continue;
+				});
+			TestTrue(TEXT("2 Vector overlaps"), VectorConflict == EAuthorityConflict::Allowed);
+		}
+
+		// Requestor replicating Float and requests Vector
+		// Existing replicating Vector
+		{
+			const FTestGroundTruth GroundTruth { ClientStream_FloatOnly, ClientStream_VectorOnly, true, true };
+			int32 NumCalls = 0;
+			const EAuthorityConflict FloatConflict = EnumerateAuthorityConflicts(RequestingClientId, ObjectPath, { Float, Vector, VectorX }, GroundTruth,
+				[this, &TestStreamId, &Vector, &VectorX, &NumCalls](const FGuid& ClientId, const FGuid& StreamId, const FConcertPropertyChain& ConflictingProperty)
+				{
+					++NumCalls;
+					TestEqual(TEXT("Float conflict > Client"), ClientId, ExistingClientId);
+					TestEqual(TEXT("Float conflict > Stream"), StreamId, TestStreamId);
+					TestTrue(TEXT("Correct roperty"), Vector == ConflictingProperty || VectorX == ConflictingProperty);
+					return EBreakBehavior::Continue;
+				});
+			TestTrue(TEXT("3 Float overlaps"), FloatConflict == EAuthorityConflict::Conflict);
+			TestEqual(TEXT("3 Float > Exactly 2 conflicting properties"), NumCalls, 2);
+		}
+		
+		return true;
+	}
 }
