Remove unneeded wrapper request types

diff --git a/Engine/Plugins/Developer/Concert/ConcertApp/MultiUserClient/Source/MultiUserClient/Private/Replication/Authority/AuthorityChangeTracker.cpp b/Engine/Plugins/Developer/Concert/ConcertApp/MultiUserClient/Source/MultiUserClient/Private/Replication/Authority/AuthorityChangeTracker.cpp
--- a/Engine/Plugins/Developer/Concert/ConcertApp/MultiUserClient/Source/MultiUserClient/Private/Replication/Authority/AuthorityChangeTracker.cpp
+++ b/Engine/Plugins/Developer/Concert/ConcertApp/MultiUserClient/Source/MultiUserClient/Private/Replication/Authority/AuthorityChangeTracker.cpp
@@ -88,14 +88,14 @@ namespace UE::MultiUserClient
 		return AuthorityCache.CanClientTakeAuthorityAfterSubmission(ObjectPath, ClientId);
 	}
 
-	TOptional<ConcertSyncClient::Replication::FAuthorityChangeRequest> FAuthorityChangeTracker::BuildChangeRequest(const FGuid& StreamId) const
+	TOptional<FConcertReplication_ChangeAuthority_Request> FAuthorityChangeTracker::BuildChangeRequest(const FGuid& StreamId) const
 	{
 		if (NewAuthorityStates.IsEmpty())
 		{
 			return {};
 		}
 		
-		ConcertSyncClient::Replication::FAuthorityChangeRequest ChangeRequest;
+		FConcertReplication_ChangeAuthority_Request ChangeRequest;
 		for (const TPair<FSoftObjectPath, bool>& NewAuthorityState : NewAuthorityStates)
 		{
 			if (NewAuthorityState.Value)
diff --git a/Engine/Plugins/Developer/Concert/ConcertApp/MultiUserClient/Source/MultiUserClient/Private/Replication/Authority/AuthorityChangeTracker.h b/Engine/Plugins/Developer/Concert/ConcertApp/MultiUserClient/Source/MultiUserClient/Private/Replication/Authority/AuthorityChangeTracker.h
--- a/Engine/Plugins/Developer/Concert/ConcertApp/MultiUserClient/Source/MultiUserClient/Private/Replication/Authority/AuthorityChangeTracker.h
+++ b/Engine/Plugins/Developer/Concert/ConcertApp/MultiUserClient/Source/MultiUserClient/Private/Replication/Authority/AuthorityChangeTracker.h
@@ -46,7 +46,7 @@ namespace UE::MultiUserClient
 		EAuthorityMutability GetChangeAuthorityMutability(const FSoftObjectPath& ObjectPath) const;
 
 		/** Builds a change request from the local changes, if there are changes. */
-		TOptional<ConcertSyncClient::Replication::FAuthorityChangeRequest> BuildChangeRequest(const FGuid& StreamId) const;
+		TOptional<FConcertReplication_ChangeAuthority_Request> BuildChangeRequest(const FGuid& StreamId) const;
 
 		/** Called entries are added to NewAuthorityStates. NOT called when objects are removed. */
 		DECLARE_MULTICAST_DELEGATE(FOnAuthorityChangeMade);
diff --git a/Engine/Plugins/Developer/Concert/ConcertApp/MultiUserClient/Source/MultiUserClient/Private/Replication/Stream/StreamChangeTracker.h b/Engine/Plugins/Developer/Concert/ConcertApp/MultiUserClient/Source/MultiUserClient/Private/Replication/Stream/StreamChangeTracker.h
--- a/Engine/Plugins/Developer/Concert/ConcertApp/MultiUserClient/Source/MultiUserClient/Private/Replication/Stream/StreamChangeTracker.h
+++ b/Engine/Plugins/Developer/Concert/ConcertApp/MultiUserClient/Source/MultiUserClient/Private/Replication/Stream/StreamChangeTracker.h
@@ -14,22 +14,11 @@
 #include "Templates/UnrealTemplate.h"
 
 class IConcertSyncClient;
-
-namespace UE::ConcertSyncClient::Replication
-{
-	struct FChangeStreamResponse;
-	struct FChangeStreamRequest;
-}
+struct FConcertReplication_ChangeStream_Response;
+struct FConcertReplication_ChangeStream_Request;
 
 namespace UE::MultiUserClient
 {
-	enum class EObjectWarningFlags
-	{
-		Ok = 0,
-		MissingProperties = 1 << 0
-	};
-	ENUM_CLASS_FLAGS(EObjectWarningFlags);
-	
 	/**
 	 * Knows of the local client's registered replication streams and builds a changelist. The changelist tracks the
 	 * unconfirmed changes to the client's streams and is updates when the server confirms the change.
diff --git a/Engine/Plugins/Developer/Concert/ConcertApp/MultiUserClient/Source/MultiUserClient/Private/Replication/Submission/AutoSubmissionPolicy.cpp b/Engine/Plugins/Developer/Concert/ConcertApp/MultiUserClient/Source/MultiUserClient/Private/Replication/Submission/AutoSubmissionPolicy.cpp
--- a/Engine/Plugins/Developer/Concert/ConcertApp/MultiUserClient/Source/MultiUserClient/Private/Replication/Submission/AutoSubmissionPolicy.cpp
+++ b/Engine/Plugins/Developer/Concert/ConcertApp/MultiUserClient/Source/MultiUserClient/Private/Replication/Submission/AutoSubmissionPolicy.cpp
@@ -49,8 +49,8 @@ namespace UE::MultiUserClient
 			
 		// Even though authority request is sent after server confirms stream change, the authority request is pre-built to avoid sending changes
 		// the local client makes while we're waiting for the latent server responses.
-		TOptional<FAuthorityChangeRequest> AuthorityChangeRequest = RequestBuilder.BuildAuthorityChange();
-		TOptional<FChangeStreamRequest> StreamRequest = RequestBuilder.BuildStreamChange();
+		TOptional<FConcertReplication_ChangeAuthority_Request> AuthorityChangeRequest = RequestBuilder.BuildAuthorityChange();
+		TOptional<FConcertReplication_ChangeStream_Request> StreamRequest = RequestBuilder.BuildStreamChange();
 			
 		if (AuthorityChangeRequest || StreamRequest)
 		{
diff --git a/Engine/Plugins/Developer/Concert/ConcertApp/MultiUserClient/Source/MultiUserClient/Private/Replication/Submission/ChangeRequestBuilder.cpp b/Engine/Plugins/Developer/Concert/ConcertApp/MultiUserClient/Source/MultiUserClient/Private/Replication/Submission/ChangeRequestBuilder.cpp
--- a/Engine/Plugins/Developer/Concert/ConcertApp/MultiUserClient/Source/MultiUserClient/Private/Replication/Submission/ChangeRequestBuilder.cpp
+++ b/Engine/Plugins/Developer/Concert/ConcertApp/MultiUserClient/Source/MultiUserClient/Private/Replication/Submission/ChangeRequestBuilder.cpp
@@ -23,10 +23,10 @@ namespace UE::MultiUserClient
 		, AuthorityChangeTracker(InAuthorityChangeTracker)
 	{}
 
-	TOptional<ConcertSyncClient::Replication::FChangeStreamRequest> FChangeRequestBuilder::BuildStreamChange() const
+	TOptional<FConcertReplication_ChangeStream_Request> FChangeRequestBuilder::BuildStreamChange() const
 	{
 		using namespace UE::ConcertSyncClient::Replication;
-		TOptional<FChangeStreamRequest> StreamRequest;
+		TOptional<FConcertReplication_ChangeStream_Request> StreamRequest;
 		
 		const FStreamChangelist& Changelist = StreamChangeTracker.GetCachedDeltaChange();
 		const bool bIsChangelistEmpty = Changelist.ObjectsToPut.IsEmpty() && Changelist.ObjectsToRemove.IsEmpty();
@@ -43,7 +43,7 @@ namespace UE::MultiUserClient
 		return StreamRequest;
 	}
 
-	TOptional<ConcertSyncClient::Replication::FAuthorityChangeRequest> FChangeRequestBuilder::BuildAuthorityChange() const
+	TOptional<FConcertReplication_ChangeAuthority_Request> FChangeRequestBuilder::BuildAuthorityChange() const
 	{
 		return AuthorityChangeTracker.BuildChangeRequest(GetLocalClientStreamId());
 	}
diff --git a/Engine/Plugins/Developer/Concert/ConcertApp/MultiUserClient/Source/MultiUserClient/Private/Replication/Submission/ChangeRequestBuilder.h b/Engine/Plugins/Developer/Concert/ConcertApp/MultiUserClient/Source/MultiUserClient/Private/Replication/Submission/ChangeRequestBuilder.h
--- a/Engine/Plugins/Developer/Concert/ConcertApp/MultiUserClient/Source/MultiUserClient/Private/Replication/Submission/ChangeRequestBuilder.h
+++ b/Engine/Plugins/Developer/Concert/ConcertApp/MultiUserClient/Source/MultiUserClient/Private/Replication/Submission/ChangeRequestBuilder.h
@@ -5,11 +5,8 @@
 #include "Misc/Guid.h"
 #include "Misc/Optional.h"
 
-namespace UE::ConcertSyncClient::Replication
-{
-	struct FChangeStreamRequest;
-	struct FAuthorityChangeRequest;
-}
+struct FConcertReplication_ChangeAuthority_Request;
+struct FConcertReplication_ChangeStream_Request;
 
 namespace UE::MultiUserClient
 {
@@ -19,7 +16,7 @@ namespace UE::MultiUserClient
 	class FStreamChangeTracker;
 	
 	/**
-	 * Util that knows how to build FChangeStreamRequest and FAuthorityChangeRequest based on the local client's changes
+	 * Util that knows how to build FConcertReplication_ChangeStream_Request and FConcertReplication_ChangeAuthority_Request based on the local client's changes
 	 * and knowledge of other remote clients.
 	 */
 	class FChangeRequestBuilder
@@ -35,10 +32,10 @@ namespace UE::MultiUserClient
 			);
 
 		/** @return Valid request that can be sent to the server, if there are any local changes. */
-		TOptional<ConcertSyncClient::Replication::FChangeStreamRequest> BuildStreamChange() const;
+		TOptional<FConcertReplication_ChangeStream_Request> BuildStreamChange() const;
 
 		/** @return Valid request that can be sent to the server, if there are any local changes. */
-		TOptional<ConcertSyncClient::Replication::FAuthorityChangeRequest> BuildAuthorityChange() const;
+		TOptional<FConcertReplication_ChangeAuthority_Request> BuildAuthorityChange() const;
 
 	private:
 
diff --git a/Engine/Plugins/Developer/Concert/ConcertApp/MultiUserClient/Source/MultiUserClient/Private/Replication/Submission/Data/AuthoritySubmission.h b/Engine/Plugins/Developer/Concert/ConcertApp/MultiUserClient/Source/MultiUserClient/Private/Replication/Submission/Data/AuthoritySubmission.h
--- a/Engine/Plugins/Developer/Concert/ConcertApp/MultiUserClient/Source/MultiUserClient/Private/Replication/Submission/Data/AuthoritySubmission.h
+++ b/Engine/Plugins/Developer/Concert/ConcertApp/MultiUserClient/Source/MultiUserClient/Private/Replication/Submission/Data/AuthoritySubmission.h
@@ -48,7 +48,7 @@ namespace UE::MultiUserClient
 		EAuthoritySubmissionRequestErrorCode ErrorCode;
 		
 		/** Only valid if ErrorCode == EAuthoritySubmissionErrorCode::Success */
-		TOptional<ConcertSyncClient::Replication::FAuthorityChangeRequest> Request;
+		TOptional<FConcertReplication_ChangeAuthority_Request> Request;
 	};
 	
 	struct FSubmitAuthorityChangesResponse
@@ -56,6 +56,6 @@ namespace UE::MultiUserClient
 		EAuthoritySubmissionResponseErrorCode ErrorCode = EAuthoritySubmissionResponseErrorCode::Cancelled;
 		
 		/** Only valid if ErrorCode == EAuthoritySubmissionErrorCode::Success */
-		TOptional<ConcertSyncClient::Replication::FAuthorityChangeResponse> Response;
+		TOptional<FConcertReplication_ChangeAuthority_Response> Response;
 	};
 }
diff --git a/Engine/Plugins/Developer/Concert/ConcertApp/MultiUserClient/Source/MultiUserClient/Private/Replication/Submission/Data/StreamSubmission.h b/Engine/Plugins/Developer/Concert/ConcertApp/MultiUserClient/Source/MultiUserClient/Private/Replication/Submission/Data/StreamSubmission.h
--- a/Engine/Plugins/Developer/Concert/ConcertApp/MultiUserClient/Source/MultiUserClient/Private/Replication/Submission/Data/StreamSubmission.h
+++ b/Engine/Plugins/Developer/Concert/ConcertApp/MultiUserClient/Source/MultiUserClient/Private/Replication/Submission/Data/StreamSubmission.h
@@ -26,8 +26,8 @@ namespace UE::MultiUserClient
 	/** Contains the original request and response of a completed stream change. */
 	struct FCompletedChangeSubmission
 	{
-		ConcertSyncClient::Replication::FChangeStreamRequest Request;
-		ConcertSyncClient::Replication::FChangeStreamResponse Response;
+		FConcertReplication_ChangeStream_Request Request;
+		FConcertReplication_ChangeStream_Response Response;
 	};
 
 	/** Resulting of submitting stream changes */
diff --git a/Engine/Plugins/Developer/Concert/ConcertApp/MultiUserClient/Source/MultiUserClient/Private/Replication/Submission/ISubmissionWorkflow.h b/Engine/Plugins/Developer/Concert/ConcertApp/MultiUserClient/Source/MultiUserClient/Private/Replication/Submission/ISubmissionWorkflow.h
--- a/Engine/Plugins/Developer/Concert/ConcertApp/MultiUserClient/Source/MultiUserClient/Private/Replication/Submission/ISubmissionWorkflow.h
+++ b/Engine/Plugins/Developer/Concert/ConcertApp/MultiUserClient/Source/MultiUserClient/Private/Replication/Submission/ISubmissionWorkflow.h
@@ -9,11 +9,8 @@
 #include "Misc/Optional.h"
 #include "Templates/SharedPointer.h"
 
-namespace UE::ConcertSyncClient::Replication
-{
-	struct FChangeStreamRequest;
-	struct FAuthorityChangeRequest;
-}
+struct FConcertReplication_ChangeAuthority_Request;
+struct FConcertReplication_ChangeStream_Request;
 
 namespace UE::MultiUserClient
 {
@@ -25,8 +22,8 @@ namespace UE::MultiUserClient
 	/** At least StreamRequest or AuthorityRequest need to be valid to form a valid request */
 	struct FSubmissionParams
 	{
-		TOptional<ConcertSyncClient::Replication::FChangeStreamRequest> StreamRequest;
-		TOptional<ConcertSyncClient::Replication::FAuthorityChangeRequest> AuthorityRequest;
+		TOptional<FConcertReplication_ChangeStream_Request> StreamRequest;
+		TOptional<FConcertReplication_ChangeAuthority_Request> AuthorityRequest;
 		
 		bool IsStreamChangeEmpty() const { return !StreamRequest.IsSet() || StreamRequest->IsEmpty(); }
 		bool IsAuthorityChangeEmpty() const { return !AuthorityRequest.IsSet() || AuthorityRequest->IsEmpty(); }
diff --git a/Engine/Plugins/Developer/Concert/ConcertApp/MultiUserClient/Source/MultiUserClient/Private/Replication/Submission/MultiEdit/ReassignObjectPropertiesLogic.cpp b/Engine/Plugins/Developer/Concert/ConcertApp/MultiUserClient/Source/MultiUserClient/Private/Replication/Submission/MultiEdit/ReassignObjectPropertiesLogic.cpp
--- a/Engine/Plugins/Developer/Concert/ConcertApp/MultiUserClient/Source/MultiUserClient/Private/Replication/Submission/MultiEdit/ReassignObjectPropertiesLogic.cpp
+++ b/Engine/Plugins/Developer/Concert/ConcertApp/MultiUserClient/Source/MultiUserClient/Private/Replication/Submission/MultiEdit/ReassignObjectPropertiesLogic.cpp
@@ -22,7 +22,7 @@ namespace UE::MultiUserClient
 			/** Used to build the final changelist to apply to the target client. Maps ClientIds to the object data they had registered. */
 			TMap<FGuid, FObjectReplicationMap> OldRegisteredObjects;
 			/** Changes to make to all the other clients */
-			TMap<FGuid, ConcertSyncClient::Replication::FChangeStreamRequest> ReassignedClientRequests;
+			TMap<FGuid, FConcertReplication_ChangeStream_Request> ReassignedClientRequests;
 			/** Authority that the clients had before. This is used to determine whether to give the assigned to client authority once the stream changes have been made. */
 			TMap<FGuid, TArray<FSoftObjectPath>> ReassignedAuthority;
 		};
@@ -32,7 +32,7 @@ namespace UE::MultiUserClient
 		{
 			FObjectReassignment Result;
 			TMap<FGuid, FObjectReplicationMap>& OldRegisteredObjects = Result.OldRegisteredObjects;
-			TMap<FGuid, ConcertSyncClient::Replication::FChangeStreamRequest>& ReassignedClientRequests = Result.ReassignedClientRequests;
+			TMap<FGuid, FConcertReplication_ChangeStream_Request>& ReassignedClientRequests = Result.ReassignedClientRequests;
 			TMap<FGuid, TArray<FSoftObjectPath>>& ReassignedAuthority = Result.ReassignedAuthority;
 		
 			for (const FSoftObjectPath& ObjectPath : ObjectsToReassign)
@@ -117,13 +117,13 @@ namespace UE::MultiUserClient
 		}
 
 		/** Builds list of authority changes based on which clients we managed to change successfully. */
-		static ConcertSyncClient::Replication::FAuthorityChangeRequest MakeAuthorityRequestFrom(
+		static FConcertReplication_ChangeAuthority_Request MakeAuthorityRequestFrom(
 			const TMap<FGuid, TArray<FSoftObjectPath>>& ReassignedAuthority, 
 			const FParallelExecutionResult& ParallelExecutionResult,
 			const FGuid& TargetStreamId
 			)
 		{
-			ConcertSyncClient::Replication::FAuthorityChangeRequest ResultRequest;
+			FConcertReplication_ChangeAuthority_Request ResultRequest;
 			for (const TPair<FGuid, FSubmitStreamChangesResponse>& ChangesRequestedOnRemote : ParallelExecutionResult.StreamResponses)
 			{
 				const TOptional<FCompletedChangeSubmission>& SubmissionInfo = ChangesRequestedOnRemote.Value.SubmissionInfo;
@@ -454,7 +454,7 @@ namespace UE::MultiUserClient
 			StreamSynchronizer.GetServerState(),
 			StreamId
 			);
-		ConcertSyncClient::Replication::FChangeStreamRequest AssignedToClientRequest = bIsStreamRegistered
+		FConcertReplication_ChangeStream_Request AssignedToClientRequest = bIsStreamRegistered
 			? StreamRequestUtils::BuildChangeRequest_UpdateExistingStream(Changelist)
 			: StreamRequestUtils::BuildChangeRequest_CreateNewStream(StreamId, Changelist);
 		
diff --git a/Engine/Plugins/Developer/Concert/ConcertApp/MultiUserClient/Source/MultiUserClient/Private/Replication/Submission/Remote/RemoteSubmissionListener.cpp b/Engine/Plugins/Developer/Concert/ConcertApp/MultiUserClient/Source/MultiUserClient/Private/Replication/Submission/Remote/RemoteSubmissionListener.cpp
--- a/Engine/Plugins/Developer/Concert/ConcertApp/MultiUserClient/Source/MultiUserClient/Private/Replication/Submission/Remote/RemoteSubmissionListener.cpp
+++ b/Engine/Plugins/Developer/Concert/ConcertApp/MultiUserClient/Source/MultiUserClient/Private/Replication/Submission/Remote/RemoteSubmissionListener.cpp
@@ -79,7 +79,7 @@ namespace UE::MultiUserClient
 		}
 		
 		const TSharedPtr<ISubmissionOperation> Operation = Workflow.SubmitChanges({
-			FChangeStreamRequest { RequestData.StreamChangeRequest }, FAuthorityChangeRequest{ RequestData.AuthorityRequest } }
+			FConcertReplication_ChangeStream_Request { RequestData.StreamChangeRequest }, FConcertReplication_ChangeAuthority_Request{ RequestData.AuthorityRequest } }
 			);
 		if (!ensure(Operation))
 		{
diff --git a/Engine/Plugins/Developer/Concert/ConcertApp/MultiUserClient/Source/MultiUserClient/Private/Replication/Submission/Remote/SubmissionWorkflow_RemoteClient.cpp b/Engine/Plugins/Developer/Concert/ConcertApp/MultiUserClient/Source/MultiUserClient/Private/Replication/Submission/Remote/SubmissionWorkflow_RemoteClient.cpp
--- a/Engine/Plugins/Developer/Concert/ConcertApp/MultiUserClient/Source/MultiUserClient/Private/Replication/Submission/Remote/SubmissionWorkflow_RemoteClient.cpp
+++ b/Engine/Plugins/Developer/Concert/ConcertApp/MultiUserClient/Source/MultiUserClient/Private/Replication/Submission/Remote/SubmissionWorkflow_RemoteClient.cpp
@@ -188,7 +188,7 @@ namespace UE::MultiUserClient
 		const FSubmitAuthorityChangesRequest AuthorityRequest { EAuthoritySubmissionRequestErrorCode::Success, InProgressOperation->Parameters.AuthorityRequest };
 		if (EventData.bProcessedSuccessfully)
 		{
-			const ConcertSyncClient::Replication::FAuthorityChangeResponse ReceivedResponse { EventData.AuthorityChangeResponse };
+			const FConcertReplication_ChangeAuthority_Response ReceivedResponse { EventData.AuthorityChangeResponse };
 			const FSubmitAuthorityChangesResponse AuthorityResponse { EAuthoritySubmissionResponseErrorCode::Success, ReceivedResponse };
 			Operation.EmplaceAuthorityResponsePromise(AuthorityResponse);
 			AuthorityRequestCompletedDelegate.Broadcast(AuthorityRequest, AuthorityResponse);
diff --git a/Engine/Plugins/Developer/Concert/ConcertApp/MultiUserClient/Source/MultiUserClient/Private/Replication/Submission/SubmissionWorkflow_LocalClient.cpp b/Engine/Plugins/Developer/Concert/ConcertApp/MultiUserClient/Source/MultiUserClient/Private/Replication/Submission/SubmissionWorkflow_LocalClient.cpp
--- a/Engine/Plugins/Developer/Concert/ConcertApp/MultiUserClient/Source/MultiUserClient/Private/Replication/Submission/SubmissionWorkflow_LocalClient.cpp
+++ b/Engine/Plugins/Developer/Concert/ConcertApp/MultiUserClient/Source/MultiUserClient/Private/Replication/Submission/SubmissionWorkflow_LocalClient.cpp
@@ -32,8 +32,8 @@ namespace UE::MultiUserClient
 			return nullptr;
 		}
 		
-		const TOptional<FChangeStreamRequest>& OptionalStreamRequest = Params.StreamRequest;
-		TOptional<FAuthorityChangeRequest>& OptionalAuthorityRequest = Params.AuthorityRequest;
+		const TOptional<FConcertReplication_ChangeStream_Request>& OptionalStreamRequest = Params.StreamRequest;
+		TOptional<FConcertReplication_ChangeAuthority_Request>& OptionalAuthorityRequest = Params.AuthorityRequest;
 		
 		const bool bIsStreamChangeEmpty = Params.IsStreamChangeEmpty();
 		const bool bModifyStreams = !bIsStreamChangeEmpty;
@@ -51,7 +51,7 @@ namespace UE::MultiUserClient
 		else
 		{
 			ReplicationManager->ChangeStream(*OptionalStreamRequest)
-				.Next([this, DestructionDetection = LifetimeToken->AsWeak(), OptionalStreamRequest, AuthorityRequest = MoveTemp(OptionalAuthorityRequest)](FChangeStreamResponse&& Response) mutable
+				.Next([this, DestructionDetection = LifetimeToken->AsWeak(), OptionalStreamRequest, AuthorityRequest = MoveTemp(OptionalAuthorityRequest)](FConcertReplication_ChangeStream_Response&& Response) mutable
 				{
 					const FSubmitStreamChangesResponse SubmissionResult { EStreamSubmissionErrorCode::Success, { FCompletedChangeSubmission{*OptionalStreamRequest, Response } } };
 					// The request might execute after we're destroyed, e.g. by leaving session while request is on the way.
@@ -70,9 +70,9 @@ namespace UE::MultiUserClient
 	}
 
 	void FSubmissionWorkflow_LocalClient::OnStreamChangeCompleted(
-		const ConcertSyncClient::Replication::FChangeStreamRequest& StreamChangeRequest,
-		const ConcertSyncClient::Replication::FChangeStreamResponse& ChangeStreamResponse,
-		TOptional<ConcertSyncClient::Replication::FAuthorityChangeRequest> AuthorityChangeRequest
+		const FConcertReplication_ChangeStream_Request& StreamChangeRequest,
+		const FConcertReplication_ChangeStream_Response& ChangeStreamResponse,
+		TOptional<FConcertReplication_ChangeAuthority_Request> AuthorityChangeRequest
 		)
 	{
 		const TSharedRef<FSingleClientSubmissionOperation> Operation = InProgressOperation->Operation;
@@ -101,7 +101,7 @@ namespace UE::MultiUserClient
 		}
 	}
 
-	void FSubmissionWorkflow_LocalClient::HandlePendingAuthorityChangeRequest(TOptional<ConcertSyncClient::Replication::FAuthorityChangeRequest> AuthorityChangeRequest)
+	void FSubmissionWorkflow_LocalClient::HandlePendingAuthorityChangeRequest(TOptional<FConcertReplication_ChangeAuthority_Request> AuthorityChangeRequest)
 	{
 		using namespace ConcertSyncClient::Replication;
 		IConcertClientReplicationManager* ReplicationManager = Client->GetReplicationManager();
@@ -129,7 +129,7 @@ namespace UE::MultiUserClient
 		FSubmitAuthorityChangesRequest Request{ EAuthoritySubmissionRequestErrorCode::Success, AuthorityChangeRequest };
 		Operation->EmplaceAuthorityRequestPromise(Request);
 		ReplicationManager->RequestAuthorityChange(MoveTemp(*AuthorityChangeRequest))
-			.Next([this, Request = MoveTemp(Request), DestructionDetection = LifetimeToken->AsWeak()](FAuthorityChangeResponse&& Response)
+			.Next([this, Request = MoveTemp(Request), DestructionDetection = LifetimeToken->AsWeak()](FConcertReplication_ChangeAuthority_Response&& Response)
 			{
 				if (DestructionDetection.IsValid())
 				{
diff --git a/Engine/Plugins/Developer/Concert/ConcertApp/MultiUserClient/Source/MultiUserClient/Private/Replication/Submission/SubmissionWorkflow_LocalClient.h b/Engine/Plugins/Developer/Concert/ConcertApp/MultiUserClient/Source/MultiUserClient/Private/Replication/Submission/SubmissionWorkflow_LocalClient.h
--- a/Engine/Plugins/Developer/Concert/ConcertApp/MultiUserClient/Source/MultiUserClient/Private/Replication/Submission/SubmissionWorkflow_LocalClient.h
+++ b/Engine/Plugins/Developer/Concert/ConcertApp/MultiUserClient/Source/MultiUserClient/Private/Replication/Submission/SubmissionWorkflow_LocalClient.h
@@ -53,11 +53,11 @@ namespace UE::MultiUserClient
 
 		/** Advances the request by requesting authority. */
 		void OnStreamChangeCompleted(
-			const ConcertSyncClient::Replication::FChangeStreamRequest& StreamChangeRequest,
-			const ConcertSyncClient::Replication::FChangeStreamResponse& ChangeStreamResponse,
-			TOptional<ConcertSyncClient::Replication::FAuthorityChangeRequest> AuthorityChangeRequest
+			const FConcertReplication_ChangeStream_Request& StreamChangeRequest,
+			const FConcertReplication_ChangeStream_Response& ChangeStreamResponse,
+			TOptional<FConcertReplication_ChangeAuthority_Request> AuthorityChangeRequest
 			);
-		void HandlePendingAuthorityChangeRequest(TOptional<ConcertSyncClient::Replication::FAuthorityChangeRequest> AuthorityChangeRequest);
+		void HandlePendingAuthorityChangeRequest(TOptional<FConcertReplication_ChangeAuthority_Request> AuthorityChangeRequest);
 
 		void CleanUpSubmissionOperation();
 	};
diff --git a/Engine/Plugins/Developer/Concert/ConcertApp/MultiUserClient/Source/MultiUserClient/Private/Replication/Util/GlobalAuthorityCache.cpp b/Engine/Plugins/Developer/Concert/ConcertApp/MultiUserClient/Source/MultiUserClient/Private/Replication/Util/GlobalAuthorityCache.cpp
--- a/Engine/Plugins/Developer/Concert/ConcertApp/MultiUserClient/Source/MultiUserClient/Private/Replication/Util/GlobalAuthorityCache.cpp
+++ b/Engine/Plugins/Developer/Concert/ConcertApp/MultiUserClient/Source/MultiUserClient/Private/Replication/Util/GlobalAuthorityCache.cpp
@@ -149,12 +149,12 @@ namespace UE::MultiUserClient
 		return Result;
 	}
 
-	void FGlobalAuthorityCache::CleanseConflictsFromAuthorityRequest(ConcertSyncClient::Replication::FAuthorityChangeRequest& Request, const FGuid& SendingClient) const
+	void FGlobalAuthorityCache::CleanseConflictsFromAuthorityRequest(FConcertReplication_ChangeAuthority_Request& Request, const FGuid& SendingClient) const
 	{
 		ConcertSyncCore::Replication::AuthorityConflictUtils::CleanseConflictsFromAuthorityRequest(Request, SendingClient, *this);
 	}
 
-	void FGlobalAuthorityCache::CleanseConflictsFromStreamRequest(ConcertSyncClient::Replication::FChangeStreamRequest& Request, const FGuid& SendingClient) const
+	void FGlobalAuthorityCache::CleanseConflictsFromStreamRequest(FConcertReplication_ChangeStream_Request& Request, const FGuid& SendingClient) const
 	{
 		ConcertSyncCore::Replication::AuthorityConflictUtils::CleanseConflictsFromStreamRequest(Request, SendingClient, *this);
 	}
diff --git a/Engine/Plugins/Developer/Concert/ConcertApp/MultiUserClient/Source/MultiUserClient/Private/Replication/Util/GlobalAuthorityCache.h b/Engine/Plugins/Developer/Concert/ConcertApp/MultiUserClient/Source/MultiUserClient/Private/Replication/Util/GlobalAuthorityCache.h
--- a/Engine/Plugins/Developer/Concert/ConcertApp/MultiUserClient/Source/MultiUserClient/Private/Replication/Util/GlobalAuthorityCache.h
+++ b/Engine/Plugins/Developer/Concert/ConcertApp/MultiUserClient/Source/MultiUserClient/Private/Replication/Util/GlobalAuthorityCache.h
@@ -60,14 +60,14 @@ namespace UE::MultiUserClient
 		 * @param Request The request to clense
 		 * @param SendingClient The client that will send the request
 		 */
-		void CleanseConflictsFromAuthorityRequest(ConcertSyncClient::Replication::FAuthorityChangeRequest& Request, const FGuid& SendingClient) const;
+		void CleanseConflictsFromAuthorityRequest(FConcertReplication_ChangeAuthority_Request& Request, const FGuid& SendingClient) const;
 		/**
 		 * Removes entries from Request that would generate conflicts.
 		 * 
 		 * @param Request The request to clense
 		 * @param SendingClient The client that will send the request
 		 */
-		void CleanseConflictsFromStreamRequest(ConcertSyncClient::Replication::FChangeStreamRequest& Request, const FGuid& SendingClient) const;
+		void CleanseConflictsFromStreamRequest(FConcertReplication_ChangeStream_Request& Request, const FGuid& SendingClient) const;
 		
 		DECLARE_MULTICAST_DELEGATE_OneParam(FOnCacheChanged, const FGuid& ClientId);
 		/** Called when the cache changes for a specific client. */
diff --git a/Engine/Plugins/Developer/Concert/ConcertApp/MultiUserClient/Source/MultiUserClient/Private/Replication/Util/RegularQueryService.cpp b/Engine/Plugins/Developer/Concert/ConcertApp/MultiUserClient/Source/MultiUserClient/Private/Replication/Util/RegularQueryService.cpp
--- a/Engine/Plugins/Developer/Concert/ConcertApp/MultiUserClient/Source/MultiUserClient/Private/Replication/Util/RegularQueryService.cpp
+++ b/Engine/Plugins/Developer/Concert/ConcertApp/MultiUserClient/Source/MultiUserClient/Private/Replication/Util/RegularQueryService.cpp
@@ -113,7 +113,7 @@ namespace UE::MultiUserClient
 		if (!Request.ClientEndpointIds.IsEmpty())
 		{
 			ReplicationManager->QueryClientInfo({ Request })
-				.Next([this, WeakToken = Token.ToWeakPtr()](ConcertSyncClient::Replication::FClientQueryResponse&& Response)
+				.Next([this, WeakToken = Token.ToWeakPtr()](FConcertReplication_QueryReplicationInfo_Response&& Response)
 				{
 					const bool bCanProcessRequest = WeakToken.Pin().IsValid();
 					if (bCanProcessRequest)
diff --git a/Engine/Plugins/Developer/Concert/ConcertApp/MultiUserClient/Source/MultiUserClient/Private/Replication/Util/StreamRequestUtils.cpp b/Engine/Plugins/Developer/Concert/ConcertApp/MultiUserClient/Source/MultiUserClient/Private/Replication/Util/StreamRequestUtils.cpp
--- a/Engine/Plugins/Developer/Concert/ConcertApp/MultiUserClient/Source/MultiUserClient/Private/Replication/Util/StreamRequestUtils.cpp
+++ b/Engine/Plugins/Developer/Concert/ConcertApp/MultiUserClient/Source/MultiUserClient/Private/Replication/Util/StreamRequestUtils.cpp
@@ -7,11 +7,11 @@
 
 namespace UE::MultiUserClient::StreamRequestUtils
 {
-	ConcertSyncClient::Replication::FChangeStreamRequest BuildChangeRequest_CreateNewStream(const FGuid& StreamId, const FStreamChangelist& FromChangelist)
+	FConcertReplication_ChangeStream_Request BuildChangeRequest_CreateNewStream(const FGuid& StreamId, const FStreamChangelist& FromChangelist)
 	{
 		const TMap<FObjectInStreamID, FConcertReplication_ChangeStream_PutObject>& ObjectsToPut = FromChangelist.ObjectsToPut;
 		
-		ConcertSyncClient::Replication::FChangeStreamRequest Request;
+		FConcertReplication_ChangeStream_Request Request;
 		Request.StreamsToAdd.Emplace();
 		FReplicationStreamDescription_NetPacked& NewStream = Request.StreamsToAdd[0];
 		NewStream.BaseDescription.Identifier = StreamId;
@@ -35,7 +35,7 @@ namespace UE::MultiUserClient::StreamRequestUtils
 		return Request;
 	}
 		
-	ConcertSyncClient::Replication::FChangeStreamRequest BuildChangeRequest_UpdateExistingStream(FStreamChangelist FromChangelist)
+	FConcertReplication_ChangeStream_Request BuildChangeRequest_UpdateExistingStream(FStreamChangelist FromChangelist)
 	{
 		return { MoveTemp(FromChangelist.ObjectsToRemove), MoveTemp(FromChangelist.ObjectsToPut) };
 	}
diff --git a/Engine/Plugins/Developer/Concert/ConcertApp/MultiUserClient/Source/MultiUserClient/Private/Replication/Util/StreamRequestUtils.h b/Engine/Plugins/Developer/Concert/ConcertApp/MultiUserClient/Source/MultiUserClient/Private/Replication/Util/StreamRequestUtils.h
--- a/Engine/Plugins/Developer/Concert/ConcertApp/MultiUserClient/Source/MultiUserClient/Private/Replication/Util/StreamRequestUtils.h
+++ b/Engine/Plugins/Developer/Concert/ConcertApp/MultiUserClient/Source/MultiUserClient/Private/Replication/Util/StreamRequestUtils.h
@@ -4,10 +4,7 @@
 
 #include "Replication/IConcertClientReplicationManager.h"
 
-namespace UE::ConcertSyncClient::Replication
-{
-	struct FChangeStreamRequest;
-}
+struct FConcertReplication_ChangeStream_Request;
 
 namespace UE::MultiUserClient
 {
@@ -26,11 +23,11 @@ namespace UE::MultiUserClient::StreamRequestUtils
 	 * @param StreamId The stream that should be modified
 	 * @param FromChangelist The changes to be made
 	 */
-	ConcertSyncClient::Replication::FChangeStreamRequest BuildChangeRequest_CreateNewStream(const FGuid& StreamId, const FStreamChangelist& FromChangelist);
+	FConcertReplication_ChangeStream_Request BuildChangeRequest_CreateNewStream(const FGuid& StreamId, const FStreamChangelist& FromChangelist);
 	
 	/**
 	 * Builds a request for updating a preexisting stream.
 	 * @param FromChangelist The changes to be made
 	 */
-	ConcertSyncClient::Replication::FChangeStreamRequest BuildChangeRequest_UpdateExistingStream(FStreamChangelist FromChangelist);
+	FConcertReplication_ChangeStream_Request BuildChangeRequest_UpdateExistingStream(FStreamChangelist FromChangelist);
 }
diff --git a/Engine/Plugins/Developer/Concert/ConcertApp/MultiUserClient/Source/MultiUserClient/Private/Replication/Util/SynchronizedRequestUtils.cpp b/Engine/Plugins/Developer/Concert/ConcertApp/MultiUserClient/Source/MultiUserClient/Private/Replication/Util/SynchronizedRequestUtils.cpp
--- a/Engine/Plugins/Developer/Concert/ConcertApp/MultiUserClient/Source/MultiUserClient/Private/Replication/Util/SynchronizedRequestUtils.cpp
+++ b/Engine/Plugins/Developer/Concert/ConcertApp/MultiUserClient/Source/MultiUserClient/Private/Replication/Util/SynchronizedRequestUtils.cpp
@@ -268,7 +268,7 @@ namespace UE::MultiUserClient
 		return Operation;
 	}
 
-	TSharedPtr<IParallelSubmissionOperation> ExecuteParallelStreamChanges(FReplicationClientManager& ClientManager, TMap<FGuid, ConcertSyncClient::Replication::FChangeStreamRequest> ParallelOperations)
+	TSharedPtr<IParallelSubmissionOperation> ExecuteParallelStreamChanges(FReplicationClientManager& ClientManager, TMap<FGuid, FConcertReplication_ChangeStream_Request> ParallelOperations)
 	{
 		if (!ensure(!ParallelOperations.IsEmpty()))
 		{
@@ -277,7 +277,7 @@ namespace UE::MultiUserClient
 		
 		TMap<FGuid, FSubmissionParams> Transformed;
 		Transformed.Reserve(ParallelOperations.Num()); // check()s if Num() == 0
-		for (TPair<FGuid, ConcertSyncClient::Replication::FChangeStreamRequest>& RequestPair : ParallelOperations)
+		for (TPair<FGuid, FConcertReplication_ChangeStream_Request>& RequestPair : ParallelOperations)
 		{
 			Transformed.Emplace(RequestPair.Key, FSubmissionParams{ MoveTemp(RequestPair.Value) });
 		}
diff --git a/Engine/Plugins/Developer/Concert/ConcertApp/MultiUserClient/Source/MultiUserClient/Private/Replication/Util/SynchronizedRequestUtils.h b/Engine/Plugins/Developer/Concert/ConcertApp/MultiUserClient/Source/MultiUserClient/Private/Replication/Util/SynchronizedRequestUtils.h
--- a/Engine/Plugins/Developer/Concert/ConcertApp/MultiUserClient/Source/MultiUserClient/Private/Replication/Util/SynchronizedRequestUtils.h
+++ b/Engine/Plugins/Developer/Concert/ConcertApp/MultiUserClient/Source/MultiUserClient/Private/Replication/Util/SynchronizedRequestUtils.h
@@ -10,11 +10,6 @@
 
 struct FConcertReplication_ChangeStream_Request;
 
-namespace UE::ConcertSyncClient::Replication
-{
-	struct FChangeStreamRequest;
-}
-
 namespace UE::MultiUserClient
 {
 	class FReplicationClientManager;
@@ -55,6 +50,6 @@ namespace UE::MultiUserClient
 	 */
 	TSharedPtr<IParallelSubmissionOperation> ExecuteParallelStreamChanges(FReplicationClientManager& ClientManager, TMap<FGuid, FSubmissionParams> ParallelOperations);
 	/** Util for transforming into FSubmissionParams */
-	TSharedPtr<IParallelSubmissionOperation> ExecuteParallelStreamChanges(FReplicationClientManager& ClientManager, TMap<FGuid, ConcertSyncClient::Replication::FChangeStreamRequest> ParallelOperations);
+	TSharedPtr<IParallelSubmissionOperation> ExecuteParallelStreamChanges(FReplicationClientManager& ClientManager, TMap<FGuid, FConcertReplication_ChangeStream_Request> ParallelOperations);
 }
 
diff --git a/Engine/Plugins/Developer/Concert/ConcertSync/ConcertSyncClient/Source/ConcertSyncClient/Private/Replication/IConcertClientReplicationManager.cpp b/Engine/Plugins/Developer/Concert/ConcertSync/ConcertSyncClient/Source/ConcertSyncClient/Private/Replication/IConcertClientReplicationManager.cpp
--- a/Engine/Plugins/Developer/Concert/ConcertSync/ConcertSyncClient/Source/ConcertSyncClient/Private/Replication/IConcertClientReplicationManager.cpp
+++ b/Engine/Plugins/Developer/Concert/ConcertSync/ConcertSyncClient/Source/ConcertSyncClient/Private/Replication/IConcertClientReplicationManager.cpp
@@ -42,7 +42,7 @@ TArray<FReplicationStreamDescription> IConcertClientReplicationManager::GetRegis
 	return Result;
 }
 
-TFuture<UE::ConcertSyncClient::Replication::FAuthorityChangeResponse> IConcertClientReplicationManager::TakeAuthorityOver(TArrayView<const FSoftObjectPath> Objects)
+TFuture<FConcertReplication_ChangeAuthority_Response> IConcertClientReplicationManager::TakeAuthorityOver(TArrayView<const FSoftObjectPath> Objects)
 {
 	using namespace UE::ConcertSyncClient::Replication;
 	
@@ -51,10 +51,10 @@ TFuture<UE::ConcertSyncClient::Replication::FAuthorityChangeResponse> IConcertCl
 		UE_LOG(LogConcert, Error, TEXT("Attempted to take authority while not connected!"));
 		TMap<FSoftObjectPath, FConcertStreamArray> Result;
 		Algo::Transform(Objects, Result, [](const FSoftObjectPath& Path){ return Path; });
-		return MakeFulfilledPromise<FAuthorityChangeResponse>(FAuthorityChangeResponse{{ EReplicationResponseErrorCode::Handled ,MoveTemp(Result) }}).GetFuture();
+		return MakeFulfilledPromise<FConcertReplication_ChangeAuthority_Response>(FConcertReplication_ChangeAuthority_Response{EReplicationResponseErrorCode::Handled ,MoveTemp(Result) }).GetFuture();
 	}
 
-	FAuthorityChangeRequest Request;
+	FConcertReplication_ChangeAuthority_Request Request;
 	Private::ForEachStreamContainingObject(Objects, *this,
 		[&Request](const FSoftObjectPath& ObjectPath, const FGuid& StreamId)
 		{
@@ -67,13 +67,13 @@ TFuture<UE::ConcertSyncClient::Replication::FAuthorityChangeResponse> IConcertCl
 		// Not only does this warn about incorrect API use at runtime, this also helps debug (incorrectly written) unit tests
 		const FString ObjectsAsString = FString::JoinBy(Objects, TEXT(","), [](const FSoftObjectPath& Path){ return Path.ToString(); });
 		UE_LOG(LogConcert, Warning, TEXT("Local client did not register any stream for the given objects. This take authority request will not be sent. Objects: %s"), *ObjectsAsString);
-		return MakeFulfilledPromise<FAuthorityChangeResponse>(FAuthorityChangeResponse{}).GetFuture();
+		return MakeFulfilledPromise<FConcertReplication_ChangeAuthority_Response>(FConcertReplication_ChangeAuthority_Response{}).GetFuture();
 	}
 	
 	return RequestAuthorityChange(MoveTemp(Request));
 }
 
-TFuture<UE::ConcertSyncClient::Replication::FAuthorityChangeResponse> IConcertClientReplicationManager::ReleaseAuthorityOf(TArrayView<const FSoftObjectPath> Objects)
+TFuture<FConcertReplication_ChangeAuthority_Response> IConcertClientReplicationManager::ReleaseAuthorityOf(TArrayView<const FSoftObjectPath> Objects)
 {
 	using namespace UE::ConcertSyncClient::Replication;
 	
@@ -82,10 +82,10 @@ TFuture<UE::ConcertSyncClient::Replication::FAuthorityChangeResponse> IConcertCl
 		UE_LOG(LogConcert, Error, TEXT("Attempted to take authority while not connected!"));
 		TMap<FSoftObjectPath, FConcertStreamArray> Result;
 		Algo::Transform(Objects, Result, [](const FSoftObjectPath& Path){ return Path; });
-		return MakeFulfilledPromise<FAuthorityChangeResponse>(FAuthorityChangeResponse{{ EReplicationResponseErrorCode::Handled, MoveTemp(Result) }}).GetFuture();
+		return MakeFulfilledPromise<FConcertReplication_ChangeAuthority_Response>(FConcertReplication_ChangeAuthority_Response{ EReplicationResponseErrorCode::Handled, MoveTemp(Result) }).GetFuture();
 	}
 	
-	FAuthorityChangeRequest Request;
+	FConcertReplication_ChangeAuthority_Request Request;
 	Private::ForEachStreamContainingObject(Objects, *this,
 		[&Request](const FSoftObjectPath& ObjectPath, const FGuid& StreamId)
 		{
@@ -98,7 +98,7 @@ TFuture<UE::ConcertSyncClient::Replication::FAuthorityChangeResponse> IConcertCl
 		// Not only does this warn about incorrect API use at runtime, this also helps debug (incorrectly written) unit tests
 		const FString ObjectsAsString = FString::JoinBy(Objects, TEXT(","), [](const FSoftObjectPath& Path){ return Path.ToString(); });
 		UE_LOG(LogConcert, Warning, TEXT("Local client did not register any stream for the given objects. This release authority request will not be sent. Objects: %s"), *ObjectsAsString);
-		return MakeFulfilledPromise<FAuthorityChangeResponse>(FAuthorityChangeResponse{}).GetFuture();
+		return MakeFulfilledPromise<FConcertReplication_ChangeAuthority_Response>(FConcertReplication_ChangeAuthority_Response{}).GetFuture();
 	}
 	
 	return RequestAuthorityChange(MoveTemp(Request));
diff --git a/Engine/Plugins/Developer/Concert/ConcertSync/ConcertSyncClient/Source/ConcertSyncClient/Private/Replication/Manager/ReplicationManager.cpp b/Engine/Plugins/Developer/Concert/ConcertSync/ConcertSyncClient/Source/ConcertSyncClient/Private/Replication/Manager/ReplicationManager.cpp
--- a/Engine/Plugins/Developer/Concert/ConcertSync/ConcertSyncClient/Source/ConcertSyncClient/Private/Replication/Manager/ReplicationManager.cpp
+++ b/Engine/Plugins/Developer/Concert/ConcertSync/ConcertSyncClient/Source/ConcertSyncClient/Private/Replication/Manager/ReplicationManager.cpp
@@ -60,7 +60,7 @@ namespace UE::ConcertSyncClient::Replication
 			: EStreamEnumerationResult::NoRegisteredStreams;
 	}
 
-	TFuture<FAuthorityChangeResponse> FReplicationManager::RequestAuthorityChange(FAuthorityChangeRequest Args)
+	TFuture<FConcertReplication_ChangeAuthority_Response> FReplicationManager::RequestAuthorityChange(FConcertReplication_ChangeAuthority_Request Args)
 	{
 		if (ensureMsgf(CurrentState, TEXT("StartAcceptingJoinRequests should have been called at this point.")))
 		{
@@ -70,18 +70,18 @@ namespace UE::ConcertSyncClient::Replication
 		return RejectAll(MoveTemp(Args));
 	}
 
-	TFuture<FClientQueryResponse> FReplicationManager::QueryClientInfo(FClientQueryRequest Args)
+	TFuture<FConcertReplication_QueryReplicationInfo_Response> FReplicationManager::QueryClientInfo(FConcertReplication_QueryReplicationInfo_Request Args)
 	{
 		return ensureMsgf(CurrentState, TEXT("StartAcceptingJoinRequests should have been called at this point."))
 			? CurrentState->QueryClientInfo(MoveTemp(Args))
-			: MakeFulfilledPromise<FClientQueryResponse>().GetFuture();
+			: MakeFulfilledPromise<FConcertReplication_QueryReplicationInfo_Response>().GetFuture();
 	}
 
-	TFuture<FChangeStreamResponse> FReplicationManager::ChangeStream(FChangeStreamRequest Args)
+	TFuture<FConcertReplication_ChangeStream_Response> FReplicationManager::ChangeStream(FConcertReplication_ChangeStream_Request Args)
 	{
 		return ensureMsgf(CurrentState, TEXT("StartAcceptingJoinRequests should have been called at this point."))
 			? CurrentState->ChangeStream(MoveTemp(Args))
-			: MakeFulfilledPromise<FChangeStreamResponse>().GetFuture(); 
+			: MakeFulfilledPromise<FConcertReplication_ChangeStream_Response>().GetFuture(); 
 	}
 
 	IConcertClientReplicationManager::EAuthorityEnumerationResult FReplicationManager::ForEachClientOwnedObject(
diff --git a/Engine/Plugins/Developer/Concert/ConcertSync/ConcertSyncClient/Source/ConcertSyncClient/Private/Replication/Manager/ReplicationManager.h b/Engine/Plugins/Developer/Concert/ConcertSync/ConcertSyncClient/Source/ConcertSyncClient/Private/Replication/Manager/ReplicationManager.h
--- a/Engine/Plugins/Developer/Concert/ConcertSync/ConcertSyncClient/Source/ConcertSyncClient/Private/Replication/Manager/ReplicationManager.h
+++ b/Engine/Plugins/Developer/Concert/ConcertSync/ConcertSyncClient/Source/ConcertSyncClient/Private/Replication/Manager/ReplicationManager.h
@@ -30,9 +30,9 @@ namespace UE::ConcertSyncClient::Replication
 		virtual bool CanJoin() override;
 		virtual bool IsConnectedToReplicationSession() override;
 		virtual EStreamEnumerationResult ForEachRegisteredStream(TFunctionRef<EBreakBehavior(const FReplicationStreamDescription& Stream)> Callback) const override;
-		virtual TFuture<FAuthorityChangeResponse> RequestAuthorityChange(FAuthorityChangeRequest Args) override;
-		virtual TFuture<FClientQueryResponse> QueryClientInfo(FClientQueryRequest Args) override;
-		virtual TFuture<FChangeStreamResponse> ChangeStream(FChangeStreamRequest Args) override;
+		virtual TFuture<FConcertReplication_ChangeAuthority_Response> RequestAuthorityChange(FConcertReplication_ChangeAuthority_Request Args) override;
+		virtual TFuture<FConcertReplication_QueryReplicationInfo_Response> QueryClientInfo(FConcertReplication_QueryReplicationInfo_Request Args) override;
+		virtual TFuture<FConcertReplication_ChangeStream_Response> ChangeStream(FConcertReplication_ChangeStream_Request Args) override;
 		virtual EAuthorityEnumerationResult ForEachClientOwnedObject(TFunctionRef<EBreakBehavior(const FSoftObjectPath& Object, TSet<FGuid>&& OwningStreams)>) const override;
 		virtual TSet<FGuid> GetClientOwnedStreamsForObject(const FSoftObjectPath& ObjectPath) const override;
 		virtual FOnPreStreamsChanged& OnPreStreamsChanged() override;
diff --git a/Engine/Plugins/Developer/Concert/ConcertSync/ConcertSyncClient/Source/ConcertSyncClient/Private/Replication/Manager/ReplicationManagerState.h b/Engine/Plugins/Developer/Concert/ConcertSync/ConcertSyncClient/Source/ConcertSyncClient/Private/Replication/Manager/ReplicationManagerState.h
--- a/Engine/Plugins/Developer/Concert/ConcertSync/ConcertSyncClient/Source/ConcertSyncClient/Private/Replication/Manager/ReplicationManagerState.h
+++ b/Engine/Plugins/Developer/Concert/ConcertSync/ConcertSyncClient/Source/ConcertSyncClient/Private/Replication/Manager/ReplicationManagerState.h
@@ -23,9 +23,9 @@ namespace UE::ConcertSyncClient::Replication
 		
 		//~ Begin IConcertClientReplicationManager Interface
 		virtual EStreamEnumerationResult ForEachRegisteredStream(TFunctionRef<EBreakBehavior(const FReplicationStreamDescription& Stream)> Callback) const override { return EStreamEnumerationResult::NoRegisteredStreams; }
-		virtual TFuture<FAuthorityChangeResponse> RequestAuthorityChange(FAuthorityChangeRequest Args) override { return RejectAll(MoveTemp(Args)); }
-		virtual TFuture<FClientQueryResponse> QueryClientInfo(FClientQueryRequest Args) override { return MakeFulfilledPromise<FClientQueryResponse>().GetFuture(); }
-		virtual TFuture<FChangeStreamResponse> ChangeStream(FChangeStreamRequest Args) override { return MakeFulfilledPromise<FChangeStreamResponse>().GetFuture(); }
+		virtual TFuture<FConcertReplication_ChangeAuthority_Response> RequestAuthorityChange(FConcertReplication_ChangeAuthority_Request Args) override { return RejectAll(MoveTemp(Args)); }
+		virtual TFuture<FConcertReplication_QueryReplicationInfo_Response> QueryClientInfo(FConcertReplication_QueryReplicationInfo_Request Args) override { return MakeFulfilledPromise<FConcertReplication_QueryReplicationInfo_Response>().GetFuture(); }
+		virtual TFuture<FConcertReplication_ChangeStream_Response> ChangeStream(FConcertReplication_ChangeStream_Request Args) override { return MakeFulfilledPromise<FConcertReplication_ChangeStream_Response>().GetFuture(); }
 		virtual EAuthorityEnumerationResult ForEachClientOwnedObject(TFunctionRef<EBreakBehavior(const FSoftObjectPath& Object, TSet<FGuid>&& OwningStreams)> Callback) const override { return EAuthorityEnumerationResult::NoAuthorityAvailable; }
 		virtual TSet<FGuid> GetClientOwnedStreamsForObject(const FSoftObjectPath& ObjectPath) const override { return {}; }
 		virtual FOnPreStreamsChanged& OnPreStreamsChanged() override { return OnPreStreamsChangedDelegate; } 
diff --git a/Engine/Plugins/Developer/Concert/ConcertSync/ConcertSyncClient/Source/ConcertSyncClient/Private/Replication/Manager/ReplicationManagerState_Connected.cpp b/Engine/Plugins/Developer/Concert/ConcertSync/ConcertSyncClient/Source/ConcertSyncClient/Private/Replication/Manager/ReplicationManagerState_Connected.cpp
--- a/Engine/Plugins/Developer/Concert/ConcertSync/ConcertSyncClient/Source/ConcertSyncClient/Private/Replication/Manager/ReplicationManagerState_Connected.cpp
+++ b/Engine/Plugins/Developer/Concert/ConcertSync/ConcertSyncClient/Source/ConcertSyncClient/Private/Replication/Manager/ReplicationManagerState_Connected.cpp
@@ -102,11 +102,11 @@ namespace UE::ConcertSyncClient::Replication
 		return EStreamEnumerationResult::Iterated;
 	}
 
-	TFuture<FAuthorityChangeResponse> FReplicationManagerState_Connected::RequestAuthorityChange(FAuthorityChangeRequest Args)
+	TFuture<FConcertReplication_ChangeAuthority_Response> FReplicationManagerState_Connected::RequestAuthorityChange(FConcertReplication_ChangeAuthority_Request Args)
 	{
 		if (CVarSimulateAuthorityTimeouts.GetValueOnGameThread())
 		{
-			return MakeFulfilledPromise<FAuthorityChangeResponse>(FAuthorityChangeResponse{ EReplicationResponseErrorCode::Timeout }).GetFuture();
+			return MakeFulfilledPromise<FConcertReplication_ChangeAuthority_Response>(FConcertReplication_ChangeAuthority_Response{ EReplicationResponseErrorCode::Timeout }).GetFuture();
 		}
 		if (CVarSimulateAuthorityRejection.GetValueOnGameThread())
 		{
@@ -131,35 +131,35 @@ namespace UE::ConcertSyncClient::Replication
 					ThisPin->RevertReleasingReplicatedObjects(Args);
 				}
 
-				return FAuthorityChangeResponse { MoveTemp(Response) };
+				return FConcertReplication_ChangeAuthority_Response { MoveTemp(Response) };
 			});
 	}
 
-	TFuture<FClientQueryResponse> FReplicationManagerState_Connected::QueryClientInfo(FClientQueryRequest Args)
+	TFuture<FConcertReplication_QueryReplicationInfo_Response> FReplicationManagerState_Connected::QueryClientInfo(FConcertReplication_QueryReplicationInfo_Request Args)
 	{
 		if (CVarSimulateQueryTimeouts.GetValueOnGameThread())
 		{
-			return MakeFulfilledPromise<FClientQueryResponse>(FClientQueryResponse{ EReplicationResponseErrorCode::Timeout }).GetFuture();
+			return MakeFulfilledPromise<FConcertReplication_QueryReplicationInfo_Response>(FConcertReplication_QueryReplicationInfo_Response{ EReplicationResponseErrorCode::Timeout }).GetFuture();
 		}
 		
 		if (EnumHasAllFlags(Args.QueryFlags, EConcertQueryClientStreamFlags::SkipAuthority | EConcertQueryClientStreamFlags::SkipStreamInfo))
 		{
 			UE_LOG(LogConcert, Warning, TEXT("Request QueryClientInfo is pointless because SkipAuthority and SkipStreamInfo are both set. Returning immediately..."));
-			return MakeFulfilledPromise<FClientQueryResponse>().GetFuture();
+			return MakeFulfilledPromise<FConcertReplication_QueryReplicationInfo_Response>().GetFuture();
 		}
 		
 		return LiveSession->SendCustomRequest<FConcertReplication_QueryReplicationInfo_Request, FConcertReplication_QueryReplicationInfo_Response>(Args, LiveSession->GetSessionServerEndpointId())
 			.Next([](FConcertReplication_QueryReplicationInfo_Response&& Response)
 			{
-				return FClientQueryResponse { MoveTemp(Response) };
+				return FConcertReplication_QueryReplicationInfo_Response { MoveTemp(Response) };
 			});
 	}
 
-	TFuture<FChangeStreamResponse> FReplicationManagerState_Connected::ChangeStream(FChangeStreamRequest Args)
+	TFuture<FConcertReplication_ChangeStream_Response> FReplicationManagerState_Connected::ChangeStream(FConcertReplication_ChangeStream_Request Args)
 	{
 		if (CVarSimulateStreamChangeTimeouts.GetValueOnGameThread())
 		{
-			return MakeFulfilledPromise<FChangeStreamResponse>(FChangeStreamResponse{ EReplicationResponseErrorCode::Timeout }).GetFuture();
+			return MakeFulfilledPromise<FConcertReplication_ChangeStream_Response>(FConcertReplication_ChangeStream_Response{ EReplicationResponseErrorCode::Timeout }).GetFuture();
 		}
 		
 		// Stop replicating removed objects right now: the server will remove authority after processing this request.
@@ -180,7 +180,7 @@ namespace UE::ConcertSyncClient::Replication
 					ThisPin->RevertRemovingReplicatedObjects(Args);
 				}
 				
-				return FChangeStreamResponse { MoveTemp(Response) };
+				return FConcertReplication_ChangeStream_Response { MoveTemp(Response) };
 			});
 	}
 
@@ -248,7 +248,7 @@ namespace UE::ConcertSyncClient::Replication
 		ReplicationApplier->ProcessObjects(TimeBudget);
 	}
 
-	void FReplicationManagerState_Connected::UpdateReplicatedObjectsAfterStreamChange(const FChangeStreamRequest& Request, const FConcertReplication_ChangeStream_Response& Response)
+	void FReplicationManagerState_Connected::UpdateReplicatedObjectsAfterStreamChange(const FConcertReplication_ChangeStream_Request& Request, const FConcertReplication_ChangeStream_Response& Response)
 	{
 		OnPreStreamsChangedDelegate.Broadcast(Request, { Response });
 		ON_SCOPE_EXIT{ OnPostStreamsChangedDelegate.Broadcast(); };
@@ -285,7 +285,7 @@ namespace UE::ConcertSyncClient::Replication
 
 	namespace Private
 	{
-		static void ForEachObjectRemovedFromStreams(const FChangeStreamRequest& Request, TFunctionRef<void(const FSoftObjectPath& ObjectPath, const TArray<FGuid>& Streams)> Callback)
+		static void ForEachObjectRemovedFromStreams(const FConcertReplication_ChangeStream_Request& Request, TFunctionRef<void(const FSoftObjectPath& ObjectPath, const TArray<FGuid>& Streams)> Callback)
 		{
 			TMap<FSoftObjectPath, TArray<FGuid>> BundledRemovedObjects;
 			for (const FObjectInStreamID& RemovedObject : Request.ObjectsToRemove)
@@ -299,7 +299,7 @@ namespace UE::ConcertSyncClient::Replication
 		}
 	}
 
-	void FReplicationManagerState_Connected::HandleRemovingReplicatedObjects(const FChangeStreamRequest& Request) const
+	void FReplicationManagerState_Connected::HandleRemovingReplicatedObjects(const FConcertReplication_ChangeStream_Request& Request) const
 	{
 		Private::ForEachObjectRemovedFromStreams(Request, [this](const FSoftObjectPath& Object, const TArray<FGuid>& RemovedStreams)
 		{
@@ -307,7 +307,7 @@ namespace UE::ConcertSyncClient::Replication
 		});
 	}
 
-	void FReplicationManagerState_Connected::RevertRemovingReplicatedObjects(const FChangeStreamRequest& Request) const
+	void FReplicationManagerState_Connected::RevertRemovingReplicatedObjects(const FConcertReplication_ChangeStream_Request& Request) const
 	{
 		Private::ForEachObjectRemovedFromStreams(Request, [this](const FSoftObjectPath& Object, const TArray<FGuid>& RemovedStreams)
 		{
@@ -315,7 +315,7 @@ namespace UE::ConcertSyncClient::Replication
 		});
 	}
 
-	void FReplicationManagerState_Connected::UpdateReplicatedObjectsAfterAuthorityChange(FAuthorityChangeRequest&& Request, const FConcertReplication_ChangeAuthority_Response& Response) const
+	void FReplicationManagerState_Connected::UpdateReplicatedObjectsAfterAuthorityChange(FConcertReplication_ChangeAuthority_Request&& Request, const FConcertReplication_ChangeAuthority_Response& Response) const
 	{
 		OnPreAuthorityChangedDelegate.Broadcast(Request, { Response });
 		ON_SCOPE_EXIT{ OnPostAuthorityChangedDelegate.Broadcast(); };
@@ -326,7 +326,7 @@ namespace UE::ConcertSyncClient::Replication
 			// Request will be discarded so ...
 			FConcertStreamArray& ReplicatedStreams = TakeAuthority.Value;
 			
-			UE_CLOG(ReplicatedStreams.StreamIds.IsEmpty(), LogConcert, Warning, TEXT("Your FAuthorityChangeRequest::TakeAuthority request contained empty stream ID array for object %s"), *ReplicatedObject.ToString());
+			UE_CLOG(ReplicatedStreams.StreamIds.IsEmpty(), LogConcert, Warning, TEXT("Your FConcertReplication_ChangeAuthority_Request::TakeAuthority request contained empty stream ID array for object %s"), *ReplicatedObject.ToString());
 			const FConcertStreamArray* RejectedStreams = Response.RejectedObjects.Find(ReplicatedObject);
 			if (RejectedStreams)
 			{
@@ -345,7 +345,7 @@ namespace UE::ConcertSyncClient::Replication
 		}
 	}
 
-	void FReplicationManagerState_Connected::HandleReleasingReplicatedObjects(const FAuthorityChangeRequest& Request) const
+	void FReplicationManagerState_Connected::HandleReleasingReplicatedObjects(const FConcertReplication_ChangeAuthority_Request& Request) const
 	{
 		for (const TPair<FSoftObjectPath, FConcertStreamArray>& ReleaseAuthority : Request.ReleaseAuthority)
 		{
@@ -353,7 +353,7 @@ namespace UE::ConcertSyncClient::Replication
 		}
 	}
 
-	void FReplicationManagerState_Connected::RevertReleasingReplicatedObjects(const FAuthorityChangeRequest& Request) const
+	void FReplicationManagerState_Connected::RevertReleasingReplicatedObjects(const FConcertReplication_ChangeAuthority_Request& Request) const
 	{
 		for (const TPair<FSoftObjectPath, FConcertStreamArray>& ReleaseAuthority : Request.ReleaseAuthority)
 		{
diff --git a/Engine/Plugins/Developer/Concert/ConcertSync/ConcertSyncClient/Source/ConcertSyncClient/Private/Replication/Manager/ReplicationManagerState_Connected.h b/Engine/Plugins/Developer/Concert/ConcertSync/ConcertSyncClient/Source/ConcertSyncClient/Private/Replication/Manager/ReplicationManagerState_Connected.h
--- a/Engine/Plugins/Developer/Concert/ConcertSync/ConcertSyncClient/Source/ConcertSyncClient/Private/Replication/Manager/ReplicationManagerState_Connected.h
+++ b/Engine/Plugins/Developer/Concert/ConcertSync/ConcertSyncClient/Source/ConcertSyncClient/Private/Replication/Manager/ReplicationManagerState_Connected.h
@@ -48,9 +48,9 @@ namespace UE::ConcertSyncClient::Replication
 		virtual bool CanJoin() override { return false; }
 		virtual bool IsConnectedToReplicationSession() override { return true; }
 		virtual EStreamEnumerationResult ForEachRegisteredStream(TFunctionRef<EBreakBehavior(const FReplicationStreamDescription& Stream)> Callback) const override;
-		virtual TFuture<FAuthorityChangeResponse> RequestAuthorityChange(FAuthorityChangeRequest Args) override;
-		virtual TFuture<FClientQueryResponse> QueryClientInfo(FClientQueryRequest Args) override;
-		virtual TFuture<FChangeStreamResponse> ChangeStream(FChangeStreamRequest Args) override;
+		virtual TFuture<FConcertReplication_ChangeAuthority_Response> RequestAuthorityChange(FConcertReplication_ChangeAuthority_Request Args) override;
+		virtual TFuture<FConcertReplication_QueryReplicationInfo_Response> QueryClientInfo(FConcertReplication_QueryReplicationInfo_Request Args) override;
+		virtual TFuture<FConcertReplication_ChangeStream_Response> ChangeStream(FConcertReplication_ChangeStream_Request Args) override;
 		virtual EAuthorityEnumerationResult ForEachClientOwnedObject(TFunctionRef<EBreakBehavior(const FSoftObjectPath& Object, TSet<FGuid>&& OwningStreams)> Callback) const override;
 		virtual TSet<FGuid> GetClientOwnedStreamsForObject(const FSoftObjectPath& ObjectPath) const override;
 		//~ End IConcertClientReplicationManager Interface
@@ -119,9 +119,9 @@ namespace UE::ConcertSyncClient::Replication
 		void TickReceiver(float TimeBudget);
 		
 		/** Updates replicated objects affected by the change request. */
-		void UpdateReplicatedObjectsAfterStreamChange(const FChangeStreamRequest& Request, const FConcertReplication_ChangeStream_Response& Response);
-		void HandleRemovingReplicatedObjects(const FChangeStreamRequest& Request) const;
-		void RevertRemovingReplicatedObjects(const FChangeStreamRequest& Request) const;
+		void UpdateReplicatedObjectsAfterStreamChange(const FConcertReplication_ChangeStream_Request& Request, const FConcertReplication_ChangeStream_Response& Response);
+		void HandleRemovingReplicatedObjects(const FConcertReplication_ChangeStream_Request& Request) const;
+		void RevertRemovingReplicatedObjects(const FConcertReplication_ChangeStream_Request& Request) const;
 
 		/**
 		 * Updates the objects which should be replicated after changing authority.
@@ -129,8 +129,8 @@ namespace UE::ConcertSyncClient::Replication
 		 * @note Request is accepted as && because this function rewrites its memory when looking at rejections.
 		 * Since the request was already sent to the server it is assumed the request can just contain trash after.
 		 */
-		void UpdateReplicatedObjectsAfterAuthorityChange(FAuthorityChangeRequest&& Request, const FConcertReplication_ChangeAuthority_Response& Response) const;
-		void HandleReleasingReplicatedObjects(const FAuthorityChangeRequest& Request) const;
-		void RevertReleasingReplicatedObjects(const FAuthorityChangeRequest& Request) const;
+		void UpdateReplicatedObjectsAfterAuthorityChange(FConcertReplication_ChangeAuthority_Request&& Request, const FConcertReplication_ChangeAuthority_Response& Response) const;
+		void HandleReleasingReplicatedObjects(const FConcertReplication_ChangeAuthority_Request& Request) const;
+		void RevertReleasingReplicatedObjects(const FConcertReplication_ChangeAuthority_Request& Request) const;
 	};
 }
diff --git a/Engine/Plugins/Developer/Concert/ConcertSync/ConcertSyncClient/Source/ConcertSyncClient/Private/Replication/Manager/ReplicationManagerUtils.cpp b/Engine/Plugins/Developer/Concert/ConcertSync/ConcertSyncClient/Source/ConcertSyncClient/Private/Replication/Manager/ReplicationManagerUtils.cpp
--- a/Engine/Plugins/Developer/Concert/ConcertSync/ConcertSyncClient/Source/ConcertSyncClient/Private/Replication/Manager/ReplicationManagerUtils.cpp
+++ b/Engine/Plugins/Developer/Concert/ConcertSync/ConcertSyncClient/Source/ConcertSyncClient/Private/Replication/Manager/ReplicationManagerUtils.cpp
@@ -9,10 +9,10 @@
 
 namespace UE::ConcertSyncClient::Replication
 {
-	TFuture<FAuthorityChangeResponse> RejectAll(FAuthorityChangeRequest&& Args)
+	TFuture<FConcertReplication_ChangeAuthority_Response> RejectAll(FConcertReplication_ChangeAuthority_Request&& Args)
 	{
-		return MakeFulfilledPromise<FAuthorityChangeResponse>(
-			FAuthorityChangeResponse{{ EReplicationResponseErrorCode::Handled, MoveTemp(Args.TakeAuthority) }}
+		return MakeFulfilledPromise<FConcertReplication_ChangeAuthority_Response>(
+			FConcertReplication_ChangeAuthority_Response{ EReplicationResponseErrorCode::Handled, MoveTemp(Args.TakeAuthority) }
 			).GetFuture();
 	}
 }
diff --git a/Engine/Plugins/Developer/Concert/ConcertSync/ConcertSyncClient/Source/ConcertSyncClient/Private/Replication/Manager/ReplicationManagerUtils.h b/Engine/Plugins/Developer/Concert/ConcertSync/ConcertSyncClient/Source/ConcertSyncClient/Private/Replication/Manager/ReplicationManagerUtils.h
--- a/Engine/Plugins/Developer/Concert/ConcertSync/ConcertSyncClient/Source/ConcertSyncClient/Private/Replication/Manager/ReplicationManagerUtils.h
+++ b/Engine/Plugins/Developer/Concert/ConcertSync/ConcertSyncClient/Source/ConcertSyncClient/Private/Replication/Manager/ReplicationManagerUtils.h
@@ -4,11 +4,11 @@
 
 #include "Async/Future.h"
 
+struct FConcertReplication_ChangeAuthority_Response;
+struct FConcertReplication_ChangeAuthority_Request;
+
 namespace UE::ConcertSyncClient::Replication
 {
-	struct FAuthorityChangeRequest;
-	struct FAuthorityChangeResponse;
-
-	/** Creates a fullfilled future which rejects all items in FConcertChangeAuthority_Request::TakeAuthority. */
-	TFuture<FAuthorityChangeResponse> RejectAll(FAuthorityChangeRequest&& Args);
+	/** Creates a fulfilled future which rejects all items in FConcertChangeAuthority_Request::TakeAuthority. */
+	TFuture<FConcertReplication_ChangeAuthority_Response> RejectAll(FConcertReplication_ChangeAuthority_Request&& Args);
 };
diff --git a/Engine/Plugins/Developer/Concert/ConcertSync/ConcertSyncClient/Source/ConcertSyncClient/Public/Replication/IConcertClientReplicationManager.h b/Engine/Plugins/Developer/Concert/ConcertSync/ConcertSyncClient/Source/ConcertSyncClient/Public/Replication/IConcertClientReplicationManager.h
--- a/Engine/Plugins/Developer/Concert/ConcertSync/ConcertSyncClient/Source/ConcertSyncClient/Public/Replication/IConcertClientReplicationManager.h
+++ b/Engine/Plugins/Developer/Concert/ConcertSync/ConcertSyncClient/Source/ConcertSyncClient/Public/Replication/IConcertClientReplicationManager.h
@@ -37,14 +37,6 @@ namespace UE::ConcertSyncClient::Replication
 			, DetailedErrorMessage(MoveTemp(DetailedErrorMessage))
 		{}
 	};
-
-	// The intention here is to wrap the request in case there is some more specific meta data we want to add in the future.
-	struct FAuthorityChangeRequest : FConcertReplication_ChangeAuthority_Request {};
-	struct FAuthorityChangeResponse : FConcertReplication_ChangeAuthority_Response {};
-	struct FClientQueryRequest : FConcertReplication_QueryReplicationInfo_Request {};
-	struct FClientQueryResponse : FConcertReplication_QueryReplicationInfo_Response {};
-	struct FChangeStreamRequest : FConcertReplication_ChangeStream_Request {};
-	struct FChangeStreamResponse : FConcertReplication_ChangeStream_Response {};
 }
 
 /**
@@ -60,7 +52,7 @@ public:
 	/**
 	 * Joins a replication session.
 	 * Subsequent calls to JoinReplicationSession will fail until either the resulting TFuture returns or LeaveReplicationSession is called.
-	 * @note This future can finish on any thread (e.g. when message endpoint times out); usually it finishes game thread.  
+	 * @note This future can finish on any thread (e.g. when message endpoint times out); usually it finishes on the game thread.  
 	 */
 	virtual TFuture<UE::ConcertSyncClient::Replication::FJoinReplicatedSessionResult> JoinReplicationSession(UE::ConcertSyncClient::Replication::FJoinReplicatedSessionArgs Args) = 0;
 	/** Leaves the current replication session. */
@@ -79,7 +71,7 @@ public:
 	 * Iterates the streams the client has registered with the server.
 	 * It only makes sense to call this function the manager has joined a replication session.
 	 * @return Whether this manager is connected to a session (Iterated) or not (NoRegisteredStreams).
-	 * @note This future can finish on any thread (e.g. when message endpoint times out); usually it finishes game thread.  
+	 * @note This future can finish on any thread (e.g. when message endpoint times out); usually it finishes on the game thread.  
 	 */
 	virtual EStreamEnumerationResult ForEachRegisteredStream(TFunctionRef<EBreakBehavior(const FReplicationStreamDescription& Stream)> Callback) const = 0;
 	/** @return Whether this manager is has any registered streams (basically whether ForEachRegisteredStream returns EStreamEnumerationResult::Iterated). */
@@ -89,13 +81,13 @@ public:
 	
 	/**
 	 * Requests from the server to change the authority over some objects.
-	 * @note This future can finish on any thread (e.g. when message endpoint times out); usually it finishes game thread.  
+	 * @note This future can finish on any thread (e.g. when message endpoint times out); usually it finishes on the game thread.  
 	 */
-	virtual TFuture<UE::ConcertSyncClient::Replication::FAuthorityChangeResponse> RequestAuthorityChange(UE::ConcertSyncClient::Replication::FAuthorityChangeRequest Args) = 0;
+	virtual TFuture<FConcertReplication_ChangeAuthority_Response> RequestAuthorityChange(FConcertReplication_ChangeAuthority_Request Args) = 0;
 	/** Util function that will request authority for all streams for the given objects. */
-	TFuture<UE::ConcertSyncClient::Replication::FAuthorityChangeResponse> TakeAuthorityOver(TArrayView<const FSoftObjectPath> Objects);
+	TFuture<FConcertReplication_ChangeAuthority_Response> TakeAuthorityOver(TArrayView<const FSoftObjectPath> Objects);
 	/** Util function that will let go over all authority of the given objects. */
-	TFuture<UE::ConcertSyncClient::Replication::FAuthorityChangeResponse> ReleaseAuthorityOf(TArrayView<const FSoftObjectPath> Objects);
+	TFuture<FConcertReplication_ChangeAuthority_Response> ReleaseAuthorityOf(TArrayView<const FSoftObjectPath> Objects);
 
 	enum class EAuthorityEnumerationResult { NoAuthorityAvailable, Iterated };
 	/**
@@ -110,19 +102,19 @@ public:
 
 	/**
 	 * Requests replication info about other clients, including the streams registered and which objects they have authority over (i.e. are sending).
-	 * @note This future can finish on any thread (e.g. when message endpoint times out); usually it finishes game thread.  
+	 * @note This future can finish on any thread (e.g. when message endpoint times out); usually it finishes on the game thread.  
 	 */
-	virtual TFuture<UE::ConcertSyncClient::Replication::FClientQueryResponse> QueryClientInfo(UE::ConcertSyncClient::Replication::FClientQueryRequest Args) = 0;
+	virtual TFuture<FConcertReplication_QueryReplicationInfo_Response> QueryClientInfo(FConcertReplication_QueryReplicationInfo_Request Args) = 0;
 
 	/**
 	 * Requests to change the client's registered stream.
-	 * @note This future can finish on any thread (e.g. when message endpoint times out); usually it finishes game thread.  
+	 * @note This future can finish on any thread (e.g. when message endpoint times out); usually it finishes on the game thread.  
 	 */
-	virtual TFuture<UE::ConcertSyncClient::Replication::FChangeStreamResponse> ChangeStream(UE::ConcertSyncClient::Replication::FChangeStreamRequest Args) = 0;
+	virtual TFuture<FConcertReplication_ChangeStream_Response> ChangeStream(FConcertReplication_ChangeStream_Request Args) = 0;
 
 	DECLARE_MULTICAST_DELEGATE_TwoParams(FOnPreStreamsChanged,
-		const UE::ConcertSyncClient::Replication::FChangeStreamRequest&,
-		const UE::ConcertSyncClient::Replication::FChangeStreamResponse&
+		const FConcertReplication_ChangeStream_Request&,
+		const FConcertReplication_ChangeStream_Response&
 		);
 	/** Called right before the result of GetRegisteredStreams changes. */
 	virtual FOnPreStreamsChanged& OnPreStreamsChanged() = 0;
@@ -132,8 +124,8 @@ public:
 	virtual FOnPostStreamsChanged& OnPostStreamsChanged() = 0;
 
 	DECLARE_MULTICAST_DELEGATE_TwoParams(FOnPreAuthorityChanged,
-		const UE::ConcertSyncClient::Replication::FAuthorityChangeRequest&,
-		const UE::ConcertSyncClient::Replication::FAuthorityChangeResponse&
+		const FConcertReplication_ChangeAuthority_Request&,
+		const FConcertReplication_ChangeAuthority_Response&
 		);
 	/** Called right before GetClientOwnedObjects changes. */
 	virtual FOnPreAuthorityChanged& OnPreAuthorityChanged() = 0;
diff --git a/Engine/Plugins/Developer/Concert/ConcertSync/ConcertSyncTest/Source/ConcertSyncTest/Private/Replication/ReplicationAuthorityTests.cpp b/Engine/Plugins/Developer/Concert/ConcertSync/ConcertSyncTest/Source/ConcertSyncTest/Private/Replication/ReplicationAuthorityTests.cpp
--- a/Engine/Plugins/Developer/Concert/ConcertSync/ConcertSyncTest/Source/ConcertSyncTest/Private/Replication/ReplicationAuthorityTests.cpp
+++ b/Engine/Plugins/Developer/Concert/ConcertSync/ConcertSyncTest/Source/ConcertSyncTest/Private/Replication/ReplicationAuthorityTests.cpp
@@ -109,7 +109,7 @@ namespace UE::ConcertSyncTests::Replication::Authority
 		// 2.1 Senders takes authority over an object...
 		bool bSenderReceivedResponse = false;
 		ClientReplicationManager_Sender->TakeAuthorityOver({ TestObjectPath })
-			.Next([this, &bSenderReceivedResponse](const ConcertSyncClient::Replication::FAuthorityChangeResponse& Response) mutable
+			.Next([this, &bSenderReceivedResponse](const FConcertReplication_ChangeAuthority_Response& Response) mutable
 			{
 				bSenderReceivedResponse = true;
 				TestEqual(TEXT("No rejection taking authority"), Response.RejectedObjects.Num(), 0);
@@ -122,7 +122,7 @@ namespace UE::ConcertSyncTests::Replication::Authority
 		// 2.2 ... so does receiver, who fails taking authority over the same object
 		bool bReceiverReceivedResponse = false;
 		ClientReplicationManager_Receiver->TakeAuthorityOver({ TestObjectPath })
-			.Next([this, &TestObjectPath, &ReceiverStreamId, &bReceiverReceivedResponse](const ConcertSyncClient::Replication::FAuthorityChangeResponse& Response) mutable
+			.Next([this, &TestObjectPath, &ReceiverStreamId, &bReceiverReceivedResponse](const FConcertReplication_ChangeAuthority_Response& Response) mutable
 			{
 				bReceiverReceivedResponse = true;
 				TestEqual(TEXT("Rejected because Sender already has authority"), Response.RejectedObjects.Num(), 1);
@@ -145,7 +145,7 @@ namespace UE::ConcertSyncTests::Replication::Authority
 		// 2.3 ... then the sender lets go of authority ...
 		bSenderReceivedResponse = false;
 		ClientReplicationManager_Sender->ReleaseAuthorityOf({ TestObjectPath })
-			.Next([this, &bSenderReceivedResponse](const ConcertSyncClient::Replication::FAuthorityChangeResponse& Response) mutable
+			.Next([this, &bSenderReceivedResponse](const FConcertReplication_ChangeAuthority_Response& Response) mutable
 			{
 				bSenderReceivedResponse = true;
 				TestEqual(TEXT("No rejection releasing object"), Response.RejectedObjects.Num(), 0); 
@@ -158,7 +158,7 @@ namespace UE::ConcertSyncTests::Replication::Authority
 		// 2.4 ... and the receiver can take authority of the object
 		bReceiverReceivedResponse = false;
 		ClientReplicationManager_Receiver->TakeAuthorityOver({ TestObjectPath })
-			.Next([this, &bReceiverReceivedResponse](const ConcertSyncClient::Replication::FAuthorityChangeResponse& Response) mutable
+			.Next([this, &bReceiverReceivedResponse](const FConcertReplication_ChangeAuthority_Response& Response) mutable
 			{
 				bReceiverReceivedResponse = true;
 				TestEqual(TEXT("No rejection because the object should not be released"), Response.RejectedObjects.Num(), 0); 
@@ -184,11 +184,11 @@ namespace UE::ConcertSyncTests::Replication::Authority
 		bool bReceivedResponse = false;
 		const FSoftObjectPath SomePath(GetDefault<UTestReflectionObject>());
 		const FGuid StreamId = FGuid::NewGuid();
-		ConcertSyncClient::Replication::FAuthorityChangeRequest Request;
+		FConcertReplication_ChangeAuthority_Request Request;
 		Request.TakeAuthority.Add(SomePath, FConcertStreamArray{ .StreamIds = { StreamId } });
 		// Detail: Cannot use IConcertClientReplicationManager::TakeAuthority util here because it builds the request based on what streams were registered
 		ClientReplicationManager_Sender->RequestAuthorityChange(Request)
-			.Next([this, &bReceivedResponse, &SomePath, &StreamId](const ConcertSyncClient::Replication::FAuthorityChangeResponse& Response) mutable
+			.Next([this, &bReceivedResponse, &SomePath, &StreamId](const FConcertReplication_ChangeAuthority_Response& Response) mutable
 			{
 				bReceivedResponse = true;
 				TestTrue(TEXT("Cannot take authority over unregistered object"), Response.RejectedObjects.Contains(SomePath));
@@ -228,7 +228,7 @@ namespace UE::ConcertSyncTests::Replication::Authority
 		// 2.1 Senders takes authority over an object...
 		bool bSenderReceivedResponse = false;
 		ClientReplicationManager_Sender->TakeAuthorityOver({ TestObjectPath })
-			.Next([this, &bSenderReceivedResponse](const ConcertSyncClient::Replication::FAuthorityChangeResponse& Response) mutable
+			.Next([this, &bSenderReceivedResponse](const FConcertReplication_ChangeAuthority_Response& Response) mutable
 			{
 				bSenderReceivedResponse = true;
 				TestEqual(TEXT("No rejection taking authority (sender)"), Response.RejectedObjects.Num(), 0); 
@@ -246,7 +246,7 @@ namespace UE::ConcertSyncTests::Replication::Authority
 		// 2.3 which means the receiver is allowed to take authority
 		bool bReceiverReceivedResponse = false;
 		ClientReplicationManager_Receiver->TakeAuthorityOver({ TestObjectPath })
-			.Next([this, &bReceiverReceivedResponse](const ConcertSyncClient::Replication::FAuthorityChangeResponse& Response) mutable
+			.Next([this, &bReceiverReceivedResponse](const FConcertReplication_ChangeAuthority_Response& Response) mutable
 			{
 				bReceiverReceivedResponse = true;
 				TestEqual(TEXT("No rejection taking authority (receiver)"), Response.RejectedObjects.Num(), 0); 
@@ -274,7 +274,7 @@ namespace UE::ConcertSyncTests::Replication::Authority
 		// 2.1 Senders takes authority over an object...
 		bool bSenderReceivedResponse = false;
 		ClientReplicationManager_Sender->TakeAuthorityOver({ TestObjectPath })
-			.Next([this, &bSenderReceivedResponse](const ConcertSyncClient::Replication::FAuthorityChangeResponse& Response) mutable
+			.Next([this, &bSenderReceivedResponse](const FConcertReplication_ChangeAuthority_Response& Response) mutable
 			{
 				bSenderReceivedResponse = true;
 				TestEqual(TEXT("No rejection taking authority (sender)"), Response.RejectedObjects.Num(), 0); 
@@ -292,7 +292,7 @@ namespace UE::ConcertSyncTests::Replication::Authority
 		// 2.3 which means the receiver is allowed to take authority
 		bool bReceiverReceivedResponse = false;
 		ClientReplicationManager_Receiver->TakeAuthorityOver({ TestObjectPath })
-			.Next([this, &bReceiverReceivedResponse](const ConcertSyncClient::Replication::FAuthorityChangeResponse& Response) mutable
+			.Next([this, &bReceiverReceivedResponse](const FConcertReplication_ChangeAuthority_Response& Response) mutable
 			{
 				bReceiverReceivedResponse = true;
 				TestEqual(TEXT("No rejection taking authority (receiver)"), Response.RejectedObjects.Num(), 0); 
@@ -320,7 +320,7 @@ namespace UE::ConcertSyncTests::Replication::Authority
 		TestEqual(TEXT("No tracked objects"), BridgeMock_Sender->TrackedObjects.Num(), 0);
 		bool bSenderReceivedTakeResponse = false;
 		ClientReplicationManager_Sender->TakeAuthorityOver({ TestObject })
-			.Next([this, &bSenderReceivedTakeResponse](const ConcertSyncClient::Replication::FAuthorityChangeResponse& Response) mutable
+			.Next([this, &bSenderReceivedTakeResponse](const FConcertReplication_ChangeAuthority_Response& Response) mutable
 			{
 				bSenderReceivedTakeResponse = true;
 				TestEqual(TEXT("No rejection taking authority"), Response.RejectedObjects.Num(), 0); 
@@ -335,7 +335,7 @@ namespace UE::ConcertSyncTests::Replication::Authority
 		// 4. Release authority
 		bool bSenderReceivedReleaseResponse = false;
 		ClientReplicationManager_Sender->ReleaseAuthorityOf({ TestObject })
-			.Next([this, &bSenderReceivedReleaseResponse](const ConcertSyncClient::Replication::FAuthorityChangeResponse& Response) mutable
+			.Next([this, &bSenderReceivedReleaseResponse](const FConcertReplication_ChangeAuthority_Response& Response) mutable
 			{
 				bSenderReceivedReleaseResponse = true;
 				TestEqual(TEXT("No rejection releasing authority"), Response.RejectedObjects.Num(), 0); 
diff --git a/Engine/Plugins/Developer/Concert/ConcertSync/ConcertSyncTest/Source/ConcertSyncTest/Private/Replication/ReplicationSendReceiveTests.cpp b/Engine/Plugins/Developer/Concert/ConcertSync/ConcertSyncTest/Source/ConcertSyncTest/Private/Replication/ReplicationSendReceiveTests.cpp
--- a/Engine/Plugins/Developer/Concert/ConcertSync/ConcertSyncTest/Source/ConcertSyncTest/Private/Replication/ReplicationSendReceiveTests.cpp
+++ b/Engine/Plugins/Developer/Concert/ConcertSync/ConcertSyncTest/Source/ConcertSyncTest/Private/Replication/ReplicationSendReceiveTests.cpp
@@ -27,7 +27,7 @@ namespace UE::ConcertSyncTests::Replication::SendReceiveFlow
 	{
 		Test.SetUpClientAndServer();
 		Test.ClientReplicationManager_Sender->TakeAuthorityOver({ Test.TestObject })
-			.Next([&Test](ConcertSyncClient::Replication::FAuthorityChangeResponse&& Response)
+			.Next([&Test](FConcertReplication_ChangeAuthority_Response&& Response)
 			{
 				if (!Response.RejectedObjects.IsEmpty())
 				{
@@ -75,11 +75,11 @@ namespace UE::ConcertSyncTests::Replication::SendReceiveFlow
 		const TOptional<FConcertPropertyChain> VectorPropertyChain = FConcertPropertyChain::CreateFromPath(*UTestReflectionObject::StaticClass(), { TEXT("Vector") });
 		VectorOnlySelection.ReplicatedProperties.Add(*VectorPropertyChain);
 		
-		ConcertSyncClient::Replication::FChangeStreamRequest Request;
+		FConcertReplication_ChangeStream_Request Request;
 		Request.ObjectsToPut.Add(FObjectInStreamID{ Test.SenderStreamId, Test.TestObject }, FConcertReplication_ChangeStream_PutObject{ VectorOnlySelection });
 		bool bReceivedChangeStreamResponse = false;
 		Test.ClientReplicationManager_Sender->ChangeStream(Request)
-			.Next([&Test, &bReceivedChangeStreamResponse](ConcertSyncClient::Replication::FChangeStreamResponse&& Response)
+			.Next([&Test, &bReceivedChangeStreamResponse](FConcertReplication_ChangeStream_Response&& Response)
 			{
 				bReceivedChangeStreamResponse = true;
 				Test.TestTrue(TEXT("Changed Stream"), Response.IsSuccess());
diff --git a/Engine/Plugins/Developer/Concert/ConcertSync/ConcertSyncTest/Source/ConcertSyncTest/Private/Replication/ReplicationStreamRequestTests.cpp b/Engine/Plugins/Developer/Concert/ConcertSync/ConcertSyncTest/Source/ConcertSyncTest/Private/Replication/ReplicationStreamRequestTests.cpp
--- a/Engine/Plugins/Developer/Concert/ConcertSync/ConcertSyncTest/Source/ConcertSyncTest/Private/Replication/ReplicationStreamRequestTests.cpp
+++ b/Engine/Plugins/Developer/Concert/ConcertSync/ConcertSyncTest/Source/ConcertSyncTest/Private/Replication/ReplicationStreamRequestTests.cpp
@@ -28,9 +28,9 @@ namespace UE::ConcertSyncTests::Replication::Stream
 		
 		const FGuid SenderEndpointId = Client_Sender->ClientSessionMock->GetSessionClientEndpointId();
 		const FSoftObjectPath TestObjectPath{ TestObject };
-		ConcertSyncClient::Replication::FClientQueryRequest Request;
+		FConcertReplication_QueryReplicationInfo_Request Request;
 		Request.ClientEndpointIds = { SenderEndpointId };
-		auto TestReplicationMapContent = [this, &SenderEndpointId](const ConcertSyncClient::Replication::FClientQueryResponse& Response)
+		auto TestReplicationMapContent = [this, &SenderEndpointId](const FConcertReplication_QueryReplicationInfo_Response& Response)
 		{
 			TestTrue(TEXT("ErrorCode == Handled"), Response.ErrorCode == EReplicationResponseErrorCode::Handled);
 			
@@ -59,7 +59,7 @@ namespace UE::ConcertSyncTests::Replication::Stream
 		// 2.1 Request before taking authority
 		bool bReceivedFirstQueryResponse = false;
 		ClientReplicationManager_Receiver->QueryClientInfo(Request)
-			.Next([this, &SenderEndpointId, &TestReplicationMapContent, &bReceivedFirstQueryResponse](ConcertSyncClient::Replication::FClientQueryResponse&& Response) mutable
+			.Next([this, &SenderEndpointId, &TestReplicationMapContent, &bReceivedFirstQueryResponse](FConcertReplication_QueryReplicationInfo_Response&& Response) mutable
 			{
 				bReceivedFirstQueryResponse = true;
 				TestTrue(TEXT("ErrorCode == Handled"), Response.ErrorCode == EReplicationResponseErrorCode::Handled);
@@ -75,7 +75,7 @@ namespace UE::ConcertSyncTests::Replication::Stream
 		// 2.2 Request after taking authority
 		bool bSenderReceivedResponse = false;
 		ClientReplicationManager_Sender->TakeAuthorityOver({ TestObjectPath })
-			.Next([this, &bSenderReceivedResponse](const ConcertSyncClient::Replication::FAuthorityChangeResponse& Response) mutable
+			.Next([this, &bSenderReceivedResponse](const FConcertReplication_ChangeAuthority_Response& Response) mutable
 			{
 				bSenderReceivedResponse = true;
 				TestEqual(TEXT("No rejection taking authority"), Response.RejectedObjects.Num(), 0); 
@@ -85,7 +85,7 @@ namespace UE::ConcertSyncTests::Replication::Stream
 
 		bool bReceivedSecondQueryResponse = false;
 		ClientReplicationManager_Receiver->QueryClientInfo(Request)
-			.Next([this, &SenderEndpointId, &TestReplicationMapContent, &TestObjectPath, &bReceivedSecondQueryResponse](ConcertSyncClient::Replication::FClientQueryResponse&& Response) mutable
+			.Next([this, &SenderEndpointId, &TestReplicationMapContent, &TestObjectPath, &bReceivedSecondQueryResponse](FConcertReplication_QueryReplicationInfo_Response&& Response) mutable
 			{
 				bReceivedSecondQueryResponse = true;
 				TestTrue(TEXT("ErrorCode == Handled"), Response.ErrorCode == EReplicationResponseErrorCode::Handled);
@@ -111,7 +111,7 @@ namespace UE::ConcertSyncTests::Replication::Stream
 		bool bReceivedResponse_SkipStreamInfo = false;
 		Request.QueryFlags = EConcertQueryClientStreamFlags::SkipStreamInfo;
 		ClientReplicationManager_Receiver->QueryClientInfo(Request)
-			.Next([this, &SenderEndpointId, &bReceivedResponse_SkipStreamInfo](ConcertSyncClient::Replication::FClientQueryResponse&& Response) mutable
+			.Next([this, &SenderEndpointId, &bReceivedResponse_SkipStreamInfo](FConcertReplication_QueryReplicationInfo_Response&& Response) mutable
 			{
 				bReceivedResponse_SkipStreamInfo = true;
 				TestTrue(TEXT("ErrorCode == Handled"), Response.ErrorCode == EReplicationResponseErrorCode::Handled);
@@ -126,7 +126,7 @@ namespace UE::ConcertSyncTests::Replication::Stream
 		bool bReceivedResponse_SkipProperties = false;
 		Request.QueryFlags = EConcertQueryClientStreamFlags::SkipProperties;
 		ClientReplicationManager_Receiver->QueryClientInfo(Request)
-			.Next([this, &SenderEndpointId, &TestObjectPath, &bReceivedResponse_SkipProperties](ConcertSyncClient::Replication::FClientQueryResponse&& Response) mutable
+			.Next([this, &SenderEndpointId, &TestObjectPath, &bReceivedResponse_SkipProperties](FConcertReplication_QueryReplicationInfo_Response&& Response) mutable
 			{
 				bReceivedResponse_SkipProperties = true;
 				const FReplicationClientQueriedInfo* Info = Response.ClientInfo.Find(SenderEndpointId);
@@ -151,7 +151,7 @@ namespace UE::ConcertSyncTests::Replication::Stream
 		bool bReceivedResponse_SkipAuthority = false;
 		Request.QueryFlags = EConcertQueryClientStreamFlags::SkipAuthority;
         ClientReplicationManager_Receiver->QueryClientInfo(Request)
-        	.Next([this, &SenderEndpointId, &bReceivedResponse_SkipAuthority](ConcertSyncClient::Replication::FClientQueryResponse&& Response) mutable
+        	.Next([this, &SenderEndpointId, &bReceivedResponse_SkipAuthority](FConcertReplication_QueryReplicationInfo_Response&& Response) mutable
         	{
         		bReceivedResponse_SkipAuthority = true;
 				TestTrue(TEXT("ErrorCode == Handled"), Response.ErrorCode == EReplicationResponseErrorCode::Handled);
@@ -234,13 +234,13 @@ namespace UE::ConcertSyncTests::Replication::Stream
 		}
 
 		/** Sends a ChangeStream request and validates that ExpectedStreams is equal to the registered streams. */
-		void ChangeStreamForSenderClientAndValidate(const FString& InTestName, const ConcertSyncClient::Replication::FChangeStreamRequest& Request, const TArray<FSharedReplicationStreamDescription>& ExpectedStreams)
+		void ChangeStreamForSenderClientAndValidate(const FString& InTestName, const FConcertReplication_ChangeStream_Request& Request, const TArray<FSharedReplicationStreamDescription>& ExpectedStreams)
 		{
 			using namespace ConcertSyncClient::Replication;
 			
 			bool bModifiedInitialStream = false;
 			ClientReplicationManager_Sender->ChangeStream(Request)
-				.Next([this, &InTestName, &bModifiedInitialStream](FChangeStreamResponse&& Response)
+				.Next([this, &InTestName, &bModifiedInitialStream](FConcertReplication_ChangeStream_Response&& Response)
 				{
 					bModifiedInitialStream = true;
 					TestTrue(TEXT("ErrorCode == Handled"), Response.ErrorCode == EReplicationResponseErrorCode::Handled);
@@ -258,7 +258,7 @@ namespace UE::ConcertSyncTests::Replication::Stream
 			bool bQueriedClientInfo = false;
 			const FGuid SenderId = Client_Sender->ClientSessionMock->GetSessionClientEndpointId();
 			ClientReplicationManager_Receiver->QueryClientInfo({{ { SenderId } }})
-				.Next([this, &InTestName, &bQueriedClientInfo, &Streams, &SenderId](FClientQueryResponse&& Response)
+				.Next([this, &InTestName, &bQueriedClientInfo, &Streams, &SenderId](FConcertReplication_QueryReplicationInfo_Response&& Response)
 				{
 					bQueriedClientInfo = true;
 					TestTrue(TEXT("ErrorCode == Handled"), Response.ErrorCode == EReplicationResponseErrorCode::Handled);
@@ -290,23 +290,23 @@ namespace UE::ConcertSyncTests::Replication::Stream
 		// 2.1 Modify existing stream
 		FConcertPropertySelection& ChangedInitialProperties = GetPropertySelection(InitialStream, *TestObject);
 		AddVectorProperty(ChangedInitialProperties);
-		FChangeStreamRequest ModifyRequest;
+		FConcertReplication_ChangeStream_Request ModifyRequest;
 		ModifyRequest.ObjectsToPut.Add({ InitialStreamId, TestObject }, { ChangedInitialProperties });
 		ChangeStreamForSenderClientAndValidate(TEXT("ModifyRequest"), ModifyRequest, { InitialStream.BaseDescription });
 
 		// 2.2 Add new stream
 		auto[DynamicStreamId, DynamicStream] = CreateFloatPropertyStream(*TestObject);
-		FChangeStreamRequest DynamicCreationRequest;
+		FConcertReplication_ChangeStream_Request DynamicCreationRequest;
 		DynamicCreationRequest.StreamsToAdd.Add(DynamicStream.Pack());
 		ChangeStreamForSenderClientAndValidate(TEXT("DynamicCreationRequest"), DynamicCreationRequest, { InitialStream.BaseDescription, DynamicStream.BaseDescription });
 
 		// 2.3 Remove stream
-		FChangeStreamRequest RemoveStreamRequest;
+		FConcertReplication_ChangeStream_Request RemoveStreamRequest;
 		RemoveStreamRequest.StreamsToRemove.Add(InitialStreamId);
 		ChangeStreamForSenderClientAndValidate(TEXT("RemoveStreamRequest"), RemoveStreamRequest, { DynamicStream.BaseDescription });
 
 		// 2.4 Remove object > removes stream implicitly
-		FChangeStreamRequest RemoveObjectRequest;
+		FConcertReplication_ChangeStream_Request RemoveObjectRequest;
 		RemoveObjectRequest.ObjectsToRemove.Add({ DynamicStreamId, TestObject });
 		ChangeStreamForSenderClientAndValidate(TEXT("RemoveObjectRequest"), RemoveObjectRequest, {});
 		
@@ -325,21 +325,21 @@ namespace UE::ConcertSyncTests::Replication::Stream
 		
 		// 2.1 Add stream with float property
 		auto[FloatStreamID, FloatStream] = CreateFloatPropertyStream(*TestObject);
-		FChangeStreamRequest CreateFloatStreamRequest;
+		FConcertReplication_ChangeStream_Request CreateFloatStreamRequest;
 		CreateFloatStreamRequest.StreamsToAdd.Add(FloatStream.Pack());
 		ChangeStreamForSenderClientAndValidate(TEXT("CreateFloatStreamRequest"), CreateFloatStreamRequest, { FloatStream.BaseDescription });
 		// 2.2 Add a stream with vector property
 		auto[VectorFloatStreamID, VectorFloatStream] = CreateVectorPropertyStream(*TestObject);
-		FChangeStreamRequest CreateVectorFloatStreamRequest; 
+		FConcertReplication_ChangeStream_Request CreateVectorFloatStreamRequest; 
 		CreateVectorFloatStreamRequest.StreamsToAdd.Add(VectorFloatStream.Pack());
 		ChangeStreamForSenderClientAndValidate(TEXT("CreateVectorFloatStreamRequest"), CreateVectorFloatStreamRequest, { FloatStream.BaseDescription, VectorFloatStream.BaseDescription });
 
 		// 2.3 Take authority over both streams
 		bool bTookAuthority = false;
-		FAuthorityChangeRequest TakeAuthorityRequest;
+		FConcertReplication_ChangeAuthority_Request TakeAuthorityRequest;
 		TakeAuthorityRequest.TakeAuthority.Add(TestObject, FConcertStreamArray{{ FloatStreamID, VectorFloatStreamID }});
 		ClientReplicationManager_Sender->RequestAuthorityChange({ TakeAuthorityRequest })
-			.Next([this, &bTookAuthority](FAuthorityChangeResponse&& Response)
+			.Next([this, &bTookAuthority](FConcertReplication_ChangeAuthority_Response&& Response)
 			{
 				bTookAuthority = true;
 				TestTrue(TEXT("ErrorCode == Handled"), Response.ErrorCode == EReplicationResponseErrorCode::Handled);
@@ -348,7 +348,7 @@ namespace UE::ConcertSyncTests::Replication::Stream
 		TestTrue(TEXT("Authority request > Received response"), bTookAuthority);
 
 		// 2.4 Appending the overlapping float property works
-		FChangeStreamRequest AppendFloatRequest;
+		FConcertReplication_ChangeStream_Request AppendFloatRequest;
 		FConcertPropertySelection& NewSelection = GetPropertySelection(VectorFloatStream, *TestObject);
 		AddFloatProperty(NewSelection);
 		AppendFloatRequest.ObjectsToPut.Add(FObjectInStreamID{ VectorFloatStreamID, TestObject }, FConcertReplication_ChangeStream_PutObject{ NewSelection });
@@ -372,14 +372,14 @@ namespace UE::ConcertSyncTests::Replication::Stream
 		
 		// 2.1 Add stream with float property
 		auto[SenderStreamID, SenderStream] = CreateFloatPropertyStream(*TestObject);
-		FChangeStreamRequest CreateFloatStreamRequest;
+		FConcertReplication_ChangeStream_Request CreateFloatStreamRequest;
 		CreateFloatStreamRequest.StreamsToAdd.Add(SenderStream.Pack());
 		ChangeStreamForSenderClientAndValidate(TEXT("CreateFloatStreamRequest"), CreateFloatStreamRequest, { SenderStream.BaseDescription });
 
 		// 2.2 Have Sender take authority over the float property
 		bool bSenderTookAuthority = false;
 		ClientReplicationManager_Sender->TakeAuthorityOver({ TestObject })
-			.Next([this, &bSenderTookAuthority](const FAuthorityChangeResponse& Response) mutable
+			.Next([this, &bSenderTookAuthority](const FConcertReplication_ChangeAuthority_Response& Response) mutable
 			{
 				bSenderTookAuthority = true;
 				TestTrue(TEXT("ErrorCode == Handled"), Response.ErrorCode == EReplicationResponseErrorCode::Handled);
@@ -389,11 +389,11 @@ namespace UE::ConcertSyncTests::Replication::Stream
 
 		// 2.3 Have Receiver create vector stream
 		auto[ReceiverStreamID, ReceiverStream] = CreateVectorPropertyStream(*TestObject);
-		FChangeStreamRequest CreateVectorStreamRequest; 
+		FConcertReplication_ChangeStream_Request CreateVectorStreamRequest; 
 		CreateVectorStreamRequest.StreamsToAdd.Add(ReceiverStream.Pack());
 		bool bAddedVectorStream = false;
 		ClientReplicationManager_Receiver->ChangeStream(CreateVectorStreamRequest)
-			.Next([this, &bAddedVectorStream](FChangeStreamResponse&& Response)
+			.Next([this, &bAddedVectorStream](FConcertReplication_ChangeStream_Response&& Response)
 			{
 				bAddedVectorStream = true;
 				TestTrue(TEXT("ErrorCode == Handled"), Response.ErrorCode == EReplicationResponseErrorCode::Handled);
@@ -405,7 +405,7 @@ namespace UE::ConcertSyncTests::Replication::Stream
 		// 2.4 Have Receiver take authority over vector stream
 		bool bReceiverTookAuthority = false;
 		ClientReplicationManager_Receiver->TakeAuthorityOver({ TestObject })
-			.Next([this, &bReceiverTookAuthority](const ConcertSyncClient::Replication::FAuthorityChangeResponse& Response) mutable
+			.Next([this, &bReceiverTookAuthority](const FConcertReplication_ChangeAuthority_Response& Response) mutable
 			{
 				bReceiverTookAuthority = true;
 				TestTrue(TEXT("ErrorCode == Handled"), Response.ErrorCode == EReplicationResponseErrorCode::Handled);
@@ -415,7 +415,7 @@ namespace UE::ConcertSyncTests::Replication::Stream
 
 		// 2.5 Receiver is rejected from modifying the stream to include the Float property because the Sender has authority over it
 		const FObjectInStreamID TestObjectInReceiverStreamId { ReceiverStreamID, TestObject };
-		FChangeStreamRequest AddFloatToStreamRequest;
+		FConcertReplication_ChangeStream_Request AddFloatToStreamRequest;
 		FConcertPropertySelection NewSelection = GetPropertySelection(ReceiverStream, *TestObject);
 		AddFloatProperty(NewSelection);
 		AddFloatToStreamRequest.ObjectsToPut.Add(TestObjectInReceiverStreamId, { NewSelection });
@@ -424,7 +424,7 @@ namespace UE::ConcertSyncTests::Replication::Stream
 		// Server logs a warning when rejecting - avoid the test being marked with a warning.
 		AddExpectedError(TEXT("Rejecting ChangeStream"));
 		ClientReplicationManager_Receiver->ChangeStream(AddFloatToStreamRequest)
-			.Next([this, TestObject, SenderStreamID = SenderStreamID, &TestObjectInReceiverStreamId, &bReceivedResponseAddingFloat](FChangeStreamResponse&& Response)
+			.Next([this, TestObject, SenderStreamID = SenderStreamID, &TestObjectInReceiverStreamId, &bReceivedResponseAddingFloat](FConcertReplication_ChangeStream_Response&& Response)
 			{
 				bReceivedResponseAddingFloat = true;
 				TestTrue(TEXT("ErrorCode == Handled"), Response.ErrorCode == EReplicationResponseErrorCode::Handled);
@@ -458,12 +458,12 @@ namespace UE::ConcertSyncTests::Replication::Stream
 		
 		// 2.1 Add stream with float property
 		auto[StreamId, Stream] = CreateFloatPropertyStream(*TestObject);
-		FChangeStreamRequest CreateStreamRequest;
+		FConcertReplication_ChangeStream_Request CreateStreamRequest;
 		CreateStreamRequest.StreamsToAdd.Add(Stream.Pack());
 		ChangeStreamForSenderClientAndValidate(TEXT("CreateStreamRequest"), CreateStreamRequest, { Stream.BaseDescription });
 
 		// 2.2 Make a request that will fail and check that no changes were made to the original stream
-		FChangeStreamRequest InvalidRequest;
+		FConcertReplication_ChangeStream_Request InvalidRequest;
 		FConcertPropertySelection NewSelection = GetPropertySelection(Stream, *TestObject);
 		AddFloatProperty(NewSelection);
 		InvalidRequest.ObjectsToPut.Add(FObjectInStreamID{ StreamId, TestObject}, FConcertReplication_ChangeStream_PutObject{ NewSelection });
@@ -474,7 +474,7 @@ namespace UE::ConcertSyncTests::Replication::Stream
 		
 		bool bReceivedChangeStreamResponse = false;
 		ClientReplicationManager_Sender->ChangeStream(InvalidRequest)
-			.Next([this, &bReceivedChangeStreamResponse](FChangeStreamResponse&& Response)
+			.Next([this, &bReceivedChangeStreamResponse](FConcertReplication_ChangeStream_Response&& Response)
 			{
 				bReceivedChangeStreamResponse = true;
 				TestTrue(TEXT("ErrorCode == Handled"), Response.ErrorCode == EReplicationResponseErrorCode::Handled);
@@ -503,20 +503,20 @@ namespace UE::ConcertSyncTests::Replication::Stream
 		
 		// 2.1 Sender takes authority
 		ClientReplicationManager_Sender->TakeAuthorityOver({ TestObject })
-			.Next([this](FAuthorityChangeResponse&& Response)
+			.Next([this](FConcertReplication_ChangeAuthority_Response&& Response)
 			{
 				TestTrue(TEXT("ErrorCode == Handled"), Response.ErrorCode == EReplicationResponseErrorCode::Handled);
 				TestTrue(TEXT("Sender > Success taking authority"), Response.RejectedObjects.IsEmpty());
 			});
 
 		// 2.2 Sender removes the object
-		FChangeStreamRequest RemoveObjectRequest;
+		FConcertReplication_ChangeStream_Request RemoveObjectRequest;
 		RemoveObjectRequest.ObjectsToRemove.Add({ SenderStreamId, TestObject });
 		ChangeStreamForSenderClientAndValidate(TEXT("Remove Object"), RemoveObjectRequest, {});
 
 		// 2.3 Receiver can now take authority since 2.2 implicitly removed authority
 		ClientReplicationManager_Receiver->TakeAuthorityOver({ TestObject })
-			.Next([this](FAuthorityChangeResponse&& Response)
+			.Next([this](FConcertReplication_ChangeAuthority_Response&& Response)
 			{
 				TestTrue(TEXT("ErrorCode == Handled"), Response.ErrorCode == EReplicationResponseErrorCode::Handled);
 				TestTrue(TEXT("Receiver > Success taking authority"), Response.RejectedObjects.IsEmpty());
@@ -540,7 +540,7 @@ namespace UE::ConcertSyncTests::Replication::Stream
 
     	// 2. Timeout request
     	ServerSession->UnregisterCustomRequestHandler<FConcertReplication_ChangeStream_Request>();
-    	ConcertSyncClient::Replication::FChangeStreamRequest Request;
+    	FConcertReplication_ChangeStream_Request Request;
     	Request.ObjectsToRemove.Add({ InitialStreamId, TestObject });
     	ClientReplicationManager_Sender->ChangeStream(Request);
 
