Move Concert tests to Editor.Concert

#rnx

diff --git a/Engine/Plugins/Developer/Concert/ConcertSync/ConcertSyncClient/Source/ConcertSyncClient/Test/ConcertClientLiveTransactionAuthorsTests.cpp b/Engine/Plugins/Developer/Concert/ConcertSync/ConcertSyncClient/Source/ConcertSyncClient/Test/ConcertClientLiveTransactionAuthorsTests.cpp
--- a/Engine/Plugins/Developer/Concert/ConcertSync/ConcertSyncClient/Source/ConcertSyncClient/Test/ConcertClientLiveTransactionAuthorsTests.cpp
+++ b/Engine/Plugins/Developer/Concert/ConcertSync/ConcertSyncClient/Source/ConcertSyncClient/Test/ConcertClientLiveTransactionAuthorsTests.cpp
@@ -157,7 +157,7 @@ void FenceTransactions(FConcertSyncClientLiveSession& InLiveSession, FConcertCli
 }
 
 /** Ensures the live transaction authors works correctly when there is no other clients connected. */
-IMPLEMENT_SIMPLE_AUTOMATION_TEST(FConcertLiveTransactionAuthorsSingleClient, "Concert.LiveTransactionAuthors.SingleClient", ConcertClientLiveTransactionAuthorsTestFlags)
+IMPLEMENT_SIMPLE_AUTOMATION_TEST(FConcertLiveTransactionAuthorsSingleClient, "Editor.Concert.LiveTransactionAuthors.SingleClient", ConcertClientLiveTransactionAuthorsTestFlags)
 
 bool FConcertLiveTransactionAuthorsSingleClient::RunTest(const FString& Parameters)
 {
@@ -196,7 +196,7 @@ bool FConcertLiveTransactionAuthorsSingleClient::RunTest(const FString& Paramete
 }
 
 /** Ensures the live transaction authors works correctly when there are many clients connected. */
-IMPLEMENT_SIMPLE_AUTOMATION_TEST(FConcertLiveTransactionAuthorsManyClients, "Concert.LiveTransactionAuthors.ManyClients", ConcertClientLiveTransactionAuthorsTestFlags)
+IMPLEMENT_SIMPLE_AUTOMATION_TEST(FConcertLiveTransactionAuthorsManyClients, "Editor.Concert.LiveTransactionAuthors.ManyClients", ConcertClientLiveTransactionAuthorsTestFlags)
 
 bool FConcertLiveTransactionAuthorsManyClients::RunTest(const FString& Parameters)
 {
@@ -310,7 +310,7 @@ bool FConcertLiveTransactionAuthorsManyClients::RunTest(const FString& Parameter
 }
 
 /** Ensures the live transaction authors works correctly when they are some transaction owned by a disconnected client. */
-IMPLEMENT_SIMPLE_AUTOMATION_TEST(FConcertLiveTransactionAuthorsDisconnectedClient, "Concert.LiveTransactionAuthors.DisconnectedClient", ConcertClientLiveTransactionAuthorsTestFlags)
+IMPLEMENT_SIMPLE_AUTOMATION_TEST(FConcertLiveTransactionAuthorsDisconnectedClient, "Editor.Concert.LiveTransactionAuthors.DisconnectedClient", ConcertClientLiveTransactionAuthorsTestFlags)
 
 bool FConcertLiveTransactionAuthorsDisconnectedClient::RunTest(const FString& Parameters)
 {
@@ -365,7 +365,7 @@ bool FConcertLiveTransactionAuthorsDisconnectedClient::RunTest(const FString& Pa
 }
 
 /** Ensures the live transaction authors works correctly when the same person is editing a package from two editors, on the same machine, concurrently. */
-IMPLEMENT_SIMPLE_AUTOMATION_TEST(FConcertLiveTransactionAuthorsClientUsingTwoEditors, "Concert.LiveTransactionAuthors.SamePersonUsingTwoEditorsConcurrently", ConcertClientLiveTransactionAuthorsTestFlags)
+IMPLEMENT_SIMPLE_AUTOMATION_TEST(FConcertLiveTransactionAuthorsClientUsingTwoEditors, "Editor.Concert.LiveTransactionAuthors.SamePersonUsingTwoEditorsConcurrently", ConcertClientLiveTransactionAuthorsTestFlags)
 
 bool FConcertLiveTransactionAuthorsClientUsingTwoEditors::RunTest(const FString& Parameters)
 {
diff --git a/Engine/Plugins/Developer/Concert/ConcertSync/ConcertSyncTest/Source/ConcertSyncTest/Private/HistoryEdition/DependencyGraphBuilderTests.cpp b/Engine/Plugins/Developer/Concert/ConcertSync/ConcertSyncTest/Source/ConcertSyncTest/Private/HistoryEdition/DependencyGraphBuilderTests.cpp
--- a/Engine/Plugins/Developer/Concert/ConcertSync/ConcertSyncTest/Source/ConcertSyncTest/Private/HistoryEdition/DependencyGraphBuilderTests.cpp
+++ b/Engine/Plugins/Developer/Concert/ConcertSync/ConcertSyncTest/Source/ConcertSyncTest/Private/HistoryEdition/DependencyGraphBuilderTests.cpp
@@ -41,7 +41,7 @@ namespace UE::ConcertSyncTests::RenameEditAndDeleteMapsFlowTest
 	 *	7 -> 5 (PackageRename)
 	 *	8 -> 5 (PackageRemoval)
 	 */
-	IMPLEMENT_SIMPLE_AUTOMATION_TEST(FRenameEditAndDeleteMapsFlowTest, "Concert.History.BuildGraph.RenameEditAndDeleteMapsFlow", EAutomationTestFlags::ApplicationContextMask | EAutomationTestFlags::EngineFilter)
+	IMPLEMENT_SIMPLE_AUTOMATION_TEST(FRenameEditAndDeleteMapsFlowTest, "Editor.Concert.History.BuildGraph.RenameEditAndDeleteMapsFlow", EAutomationTestFlags::EditorContext | EAutomationTestFlags::EngineFilter)
 	bool FRenameEditAndDeleteMapsFlowTest::RunTest(const FString& Parameters)
 	{
 		FScopedSessionDatabase SessionDatabase(*this);
@@ -190,7 +190,7 @@ namespace UE::ConcertSyncTests::DeletingAndRecreatingActorIsHardDependency
 	 *
 	 * 3 > 2 is a hard dependency (removing 2 would result in attempting to create the actor twice, which is invalid).
 	 */
-	IMPLEMENT_SIMPLE_AUTOMATION_TEST(FDeletingAndRecreatingActorIsHardDependency, "Concert.History.BuildGraph.DeletingAndRecreatingActorIsHardDependency", EAutomationTestFlags::ApplicationContextMask | EAutomationTestFlags::EngineFilter)
+	IMPLEMENT_SIMPLE_AUTOMATION_TEST(FDeletingAndRecreatingActorIsHardDependency, "Editor.Concert.History.BuildGraph.DeletingAndRecreatingActorIsHardDependency", EAutomationTestFlags::EditorContext | EAutomationTestFlags::EngineFilter)
 	bool FDeletingAndRecreatingActorIsHardDependency::RunTest(const FString& Parameters)
 	{
 		using namespace ConcertSyncCore;
@@ -307,7 +307,7 @@ namespace UE::ConcertSyncTests::IndirectPackageDependencyTest
 	 *
 	 * This is relevant because the actor's construction script may depend query data from the data asset.
 	 */
-	IMPLEMENT_SIMPLE_AUTOMATION_TEST(FIndirectPackageDependencyTest, "Concert.History.BuildGraph.IndirectPackageDependencyTest", EAutomationTestFlags::ApplicationContextMask | EAutomationTestFlags::EngineFilter)
+	IMPLEMENT_SIMPLE_AUTOMATION_TEST(FIndirectPackageDependencyTest, "Editor.Concert.History.BuildGraph.IndirectPackageDependencyTest", EAutomationTestFlags::EditorContext | EAutomationTestFlags::EngineFilter)
 	bool FIndirectPackageDependencyTest::RunTest(const FString& Parameters)
 	{
 		// TODO:
@@ -372,7 +372,7 @@ namespace UE::ConcertSyncTests::PackageEditedDependencyTest
 	 *  8 -> 7
 	 *  13 -> 12
 	 */
-	IMPLEMENT_SIMPLE_AUTOMATION_TEST(FPackageEditedDependencyTest, "Concert.History.BuildGraph.PackageEditedDependencyTest", EAutomationTestFlags::ApplicationContextMask | EAutomationTestFlags::EngineFilter)
+	IMPLEMENT_SIMPLE_AUTOMATION_TEST(FPackageEditedDependencyTest, "Editor.Concert.History.BuildGraph.PackageEditedDependencyTest", EAutomationTestFlags::EditorContext | EAutomationTestFlags::EngineFilter)
 	bool FPackageEditedDependencyTest::RunTest(const FString& Parameters)
 	{
 		using namespace ConcertSyncCore;
@@ -527,7 +527,7 @@ namespace UE::ConcertSyncTests::RenamingDependencyEdgeCases
 	 *  8 Add Bar
 	 *  9 Rename Bar to Foo
 	 */
-	IMPLEMENT_SIMPLE_AUTOMATION_TEST(FRenamingDependencyEdgeCaseTests, "Concert.History.BuildGraph.RenamingDependencyEdgeCaseTests", EAutomationTestFlags::ApplicationContextMask | EAutomationTestFlags::EngineFilter)
+	IMPLEMENT_SIMPLE_AUTOMATION_TEST(FRenamingDependencyEdgeCaseTests, "Editor.Concert.History.BuildGraph.RenamingDependencyEdgeCaseTests", EAutomationTestFlags::EditorContext | EAutomationTestFlags::EngineFilter)
 	bool FRenamingDependencyEdgeCaseTests::RunTest(const FString& Parameters)
 	{
 		using namespace ConcertSyncCore;
@@ -598,7 +598,7 @@ namespace UE::ConcertSyncTests::TransactionDependencyEdgeCases
 	
 	TArray<FActivityID> FillDatabase(FScopedSessionDatabase& SessionDatabase);
 	
-	IMPLEMENT_SIMPLE_AUTOMATION_TEST(FTransactionDependencyEdgeCaseTests, "Concert.History.BuildGraph.TransactionDependencyEdgeCaseTests", EAutomationTestFlags::ApplicationContextMask | EAutomationTestFlags::EngineFilter)
+	IMPLEMENT_SIMPLE_AUTOMATION_TEST(FTransactionDependencyEdgeCaseTests, "Editor.Concert.History.BuildGraph.TransactionDependencyEdgeCaseTests", EAutomationTestFlags::EditorContext | EAutomationTestFlags::EngineFilter)
 	bool FTransactionDependencyEdgeCaseTests::RunTest(const FString& Parameters)
 	{
 		using namespace ConcertSyncCore;
diff --git a/Engine/Plugins/Developer/Concert/ConcertSync/ConcertSyncTest/Source/ConcertSyncTest/Private/HistoryEdition/DependencyGraphTests.cpp b/Engine/Plugins/Developer/Concert/ConcertSync/ConcertSyncTest/Source/ConcertSyncTest/Private/HistoryEdition/DependencyGraphTests.cpp
--- a/Engine/Plugins/Developer/Concert/ConcertSync/ConcertSyncTest/Source/ConcertSyncTest/Private/HistoryEdition/DependencyGraphTests.cpp
+++ b/Engine/Plugins/Developer/Concert/ConcertSync/ConcertSyncTest/Source/ConcertSyncTest/Private/HistoryEdition/DependencyGraphTests.cpp
@@ -10,7 +10,7 @@
 namespace UE::ConcertSyncTests
 {
 	/** Ensures that the dependency graph API functions as intended, i.e. you add, connect and query nodes. */
-	IMPLEMENT_SIMPLE_AUTOMATION_TEST(FConcertDependencyGraphApiTest, "Concert.History.GraphApi.DependencyGraphApi", EAutomationTestFlags::ApplicationContextMask | EAutomationTestFlags::SmokeFilter)
+	IMPLEMENT_SIMPLE_AUTOMATION_TEST(FConcertDependencyGraphApiTest, "Editor.Concert.History.GraphApi.DependencyGraphApi", EAutomationTestFlags::EditorContext | EAutomationTestFlags::SmokeFilter)
 	bool FConcertDependencyGraphApiTest::RunTest(const FString& Parameters)
 	{
 		using namespace UE::ConcertSyncCore;
@@ -54,7 +54,7 @@ namespace UE::ConcertSyncTests
 	}
 
 	/** Ensures that the dependency graph API functions as intended, i.e. you add, connect and query nodes. */
-	IMPLEMENT_SIMPLE_AUTOMATION_TEST(FDisallowCyclesInGraph, "Concert.History.GraphApi.DisallowCyclesInGraph", EAutomationTestFlags::ApplicationContextMask | EAutomationTestFlags::SmokeFilter)
+	IMPLEMENT_SIMPLE_AUTOMATION_TEST(FDisallowCyclesInGraph, "Editor.Concert.History.GraphApi.DisallowCyclesInGraph", EAutomationTestFlags::EditorContext | EAutomationTestFlags::SmokeFilter)
 	bool FDisallowCyclesInGraph::RunTest(const FString& Parameters)
 	{
 		using namespace UE::ConcertSyncCore;
diff --git a/Engine/Plugins/Developer/Concert/ConcertSync/ConcertSyncTest/Source/ConcertSyncTest/Private/HistoryEdition/HistoryBottomUpAnalysisTests.cpp b/Engine/Plugins/Developer/Concert/ConcertSync/ConcertSyncTest/Source/ConcertSyncTest/Private/HistoryEdition/HistoryBottomUpAnalysisTests.cpp
--- a/Engine/Plugins/Developer/Concert/ConcertSync/ConcertSyncTest/Source/ConcertSyncTest/Private/HistoryEdition/HistoryBottomUpAnalysisTests.cpp
+++ b/Engine/Plugins/Developer/Concert/ConcertSync/ConcertSyncTest/Source/ConcertSyncTest/Private/HistoryEdition/HistoryBottomUpAnalysisTests.cpp
@@ -18,7 +18,7 @@ namespace UE::ConcertSyncTests::AnalysisTests
 	/**
 	 * Runs through a couple of use cases of unmuting an activity.
 	 */
-	IMPLEMENT_SIMPLE_AUTOMATION_TEST(FBottomUpUseCase, "Concert.History.Analysis.BottomUpUseCase", EAutomationTestFlags::ApplicationContextMask | EAutomationTestFlags::EngineFilter)
+	IMPLEMENT_SIMPLE_AUTOMATION_TEST(FBottomUpUseCase, "Editor.Concert.History.Analysis.BottomUpUseCase", EAutomationTestFlags::EditorContext | EAutomationTestFlags::EngineFilter)
 	bool FBottomUpUseCase::RunTest(const FString& Parameters)
 	{
 		using namespace RenameEditAndDeleteMapsFlowTest;
@@ -70,7 +70,7 @@ namespace UE::ConcertSyncTests::AnalysisTests
 	 * The test: unmute L.
 	 * We want R to be marked has a hard dependency.
 	 */
-	IMPLEMENT_SIMPLE_AUTOMATION_TEST(FBottomUpDiamond, "Concert.History.Analysis.BottomUpDiamond", EAutomationTestFlags::ApplicationContextMask | EAutomationTestFlags::EngineFilter)
+	IMPLEMENT_SIMPLE_AUTOMATION_TEST(FBottomUpDiamond, "Editor.Concert.History.Analysis.BottomUpDiamond", EAutomationTestFlags::EditorContext | EAutomationTestFlags::EngineFilter)
 	bool FBottomUpDiamond::RunTest(const FString& Parameters)
 	{
 		using namespace ConcertSyncCore;
diff --git a/Engine/Plugins/Developer/Concert/ConcertSync/ConcertSyncTest/Source/ConcertSyncTest/Private/HistoryEdition/HistoryTopDownAnalysisTests.cpp b/Engine/Plugins/Developer/Concert/ConcertSync/ConcertSyncTest/Source/ConcertSyncTest/Private/HistoryEdition/HistoryTopDownAnalysisTests.cpp
--- a/Engine/Plugins/Developer/Concert/ConcertSync/ConcertSyncTest/Source/ConcertSyncTest/Private/HistoryEdition/HistoryTopDownAnalysisTests.cpp
+++ b/Engine/Plugins/Developer/Concert/ConcertSync/ConcertSyncTest/Source/ConcertSyncTest/Private/HistoryEdition/HistoryTopDownAnalysisTests.cpp
@@ -58,7 +58,7 @@ namespace UE::ConcertSyncTests
 
 namespace UE::ConcertSyncTests::AnalysisTests
 {
-	IMPLEMENT_SIMPLE_AUTOMATION_TEST(FTopDownUseCase, "Concert.History.Analysis.TopDownUseCase", EAutomationTestFlags::ApplicationContextMask | EAutomationTestFlags::EngineFilter)
+	IMPLEMENT_SIMPLE_AUTOMATION_TEST(FTopDownUseCase, "Editor.Concert.History.Analysis.TopDownUseCase", EAutomationTestFlags::EditorContext | EAutomationTestFlags::EngineFilter)
 	bool FTopDownUseCase::RunTest(const FString& Parameters)
 	{
 		using namespace RenameEditAndDeleteMapsFlowTest;
@@ -132,7 +132,7 @@ namespace UE::ConcertSyncTests::AnalysisTests
 	 * The test: delete R.
 	 * We want L to be marked has a hard dependency.
 	 */
-	IMPLEMENT_SIMPLE_AUTOMATION_TEST(FTopDownDiamond, "Concert.History.Analysis.TopDownDiamond", EAutomationTestFlags::ApplicationContextMask | EAutomationTestFlags::EngineFilter)
+	IMPLEMENT_SIMPLE_AUTOMATION_TEST(FTopDownDiamond, "Editor.Concert.History.Analysis.TopDownDiamond", EAutomationTestFlags::EditorContext | EAutomationTestFlags::EngineFilter)
 	bool FTopDownDiamond::RunTest(const FString& Parameters)
 	{
 		using namespace ConcertSyncCore;
diff --git a/Engine/Plugins/Developer/Concert/ConcertSync/ConcertSyncTest/Source/ConcertSyncTest/Private/Replication/Components/ReplicationCacheTests.cpp b/Engine/Plugins/Developer/Concert/ConcertSync/ConcertSyncTest/Source/ConcertSyncTest/Private/Replication/Components/ReplicationCacheTests.cpp
--- a/Engine/Plugins/Developer/Concert/ConcertSync/ConcertSyncTest/Source/ConcertSyncTest/Private/Replication/Components/ReplicationCacheTests.cpp
+++ b/Engine/Plugins/Developer/Concert/ConcertSync/ConcertSyncTest/Source/ConcertSyncTest/Private/Replication/Components/ReplicationCacheTests.cpp
@@ -25,7 +25,7 @@ namespace UE::ConcertSyncTests
 	/**
 	 * Tests that FObjectReplicationCache works as expected.
 	 */
-	IMPLEMENT_SIMPLE_AUTOMATION_TEST(FReplicationCacheTest, "Concert.Replication.Components.ReplicationCache", EAutomationTestFlags::EditorContext | EAutomationTestFlags::EngineFilter);
+	IMPLEMENT_SIMPLE_AUTOMATION_TEST(FReplicationCacheTest, "Editor.Concert.Replication.Components.ReplicationCache", EAutomationTestFlags::EditorContext | EAutomationTestFlags::EngineFilter);
 	bool FReplicationCacheTest::RunTest(const FString& Parameters)
 	{
 		class FTestReplicationFormat : public ConcertSyncCore::IObjectReplicationFormat
diff --git a/Engine/Plugins/Developer/Concert/ConcertSync/ConcertSyncTest/Source/ConcertSyncTest/Private/Replication/Components/ReplicationChangeStreamUtilsTests.cpp b/Engine/Plugins/Developer/Concert/ConcertSync/ConcertSyncTest/Source/ConcertSyncTest/Private/Replication/Components/ReplicationChangeStreamUtilsTests.cpp
--- a/Engine/Plugins/Developer/Concert/ConcertSync/ConcertSyncTest/Source/ConcertSyncTest/Private/Replication/Components/ReplicationChangeStreamUtilsTests.cpp
+++ b/Engine/Plugins/Developer/Concert/ConcertSync/ConcertSyncTest/Source/ConcertSyncTest/Private/Replication/Components/ReplicationChangeStreamUtilsTests.cpp
@@ -19,7 +19,7 @@ namespace UE::ConcertSyncTests::ChangeStreamUtils
 	 *  - adding and removing objects
 	 *  - modify properties and class
 	 */
-	IMPLEMENT_SIMPLE_AUTOMATION_TEST(FSimpleDiffChangesTest, "Concert.Replication.Components.DiffChangesSimple", EAutomationTestFlags::EditorContext | EAutomationTestFlags::EngineFilter);
+	IMPLEMENT_SIMPLE_AUTOMATION_TEST(FSimpleDiffChangesTest, "Editor.Concert.Replication.Components.DiffChangesSimple", EAutomationTestFlags::EditorContext | EAutomationTestFlags::EngineFilter);
 	bool FSimpleDiffChangesTest::RunTest(const FString& Parameters)
 	{
 		// 1. Set up
diff --git a/Engine/Plugins/Developer/Concert/ConcertSync/ConcertSyncTest/Source/ConcertSyncTest/Private/Replication/ConcertPropertyTests.cpp b/Engine/Plugins/Developer/Concert/ConcertSync/ConcertSyncTest/Source/ConcertSyncTest/Private/Replication/ConcertPropertyTests.cpp
--- a/Engine/Plugins/Developer/Concert/ConcertSync/ConcertSyncTest/Source/ConcertSyncTest/Private/Replication/ConcertPropertyTests.cpp
+++ b/Engine/Plugins/Developer/Concert/ConcertSync/ConcertSyncTest/Source/ConcertSyncTest/Private/Replication/ConcertPropertyTests.cpp
@@ -21,7 +21,7 @@ namespace UE::ConcertSyncTest::Replication::PropertyChain
 	 *		- "normal" structs, like FVector > lists sub-properties
 	 *		- primitives and native structs > Path contains FConcertPropertyChain::InternalContainerPropertyValueName ("Value") at the end
 	 */
-	IMPLEMENT_SIMPLE_AUTOMATION_TEST(FConcertPropertyChainTests, "Concert.Replication.Data.PropertyChainTests", EAutomationTestFlags::EditorContext | EAutomationTestFlags::EngineFilter)
+	IMPLEMENT_SIMPLE_AUTOMATION_TEST(FConcertPropertyChainTests, "Editor.Concert.Replication.Data.PropertyChainTests", EAutomationTestFlags::EditorContext | EAutomationTestFlags::EngineFilter)
 	bool FConcertPropertyChainTests::RunTest(const FString& Parameters)
 	{
 		// The test is intentionally set up to make it it easy to set breakpoints albeit at the expense of increasing the amount of code
@@ -121,7 +121,7 @@ namespace UE::ConcertSyncTest::Replication::PropertyChain
 	/**
 	 * Tests that all property cases on the specially constructed example class UTestReflectionObject.
 	 */
-	IMPLEMENT_SIMPLE_AUTOMATION_TEST(FConcertForEachReplictableConcertPropertyTests, "Concert.Replication.Data.ForEachReplicatableConcertProperty", EAutomationTestFlags::EditorContext | EAutomationTestFlags::EngineFilter)
+	IMPLEMENT_SIMPLE_AUTOMATION_TEST(FConcertForEachReplictableConcertPropertyTests, "Editor.Concert.Replication.Data.ForEachReplicatableConcertProperty", EAutomationTestFlags::EditorContext | EAutomationTestFlags::EngineFilter)
 	bool FConcertForEachReplictableConcertPropertyTests::RunTest(const FString& Parameters)
 	{
 		UClass* Class = UTestReflectionObject::StaticClass();
@@ -220,7 +220,7 @@ namespace UE::ConcertSyncTest::Replication::PropertyChain
 	/**
 	 * Tests FConcertPropertyChain::Matches.
 	 */
-	IMPLEMENT_SIMPLE_AUTOMATION_TEST(FMatchConcertPropertyChainTest, "Concert.Replication.Data.MatchConcertPropertyChain", EAutomationTestFlags::EditorContext | EAutomationTestFlags::EngineFilter)
+	IMPLEMENT_SIMPLE_AUTOMATION_TEST(FMatchConcertPropertyChainTest, "Editor.Concert.Replication.Data.MatchConcertPropertyChain", EAutomationTestFlags::EditorContext | EAutomationTestFlags::EngineFilter)
 	bool FMatchConcertPropertyChainTest::RunTest(const FString& Parameters)
 	{
 		// 1. Set up
@@ -293,7 +293,7 @@ namespace UE::ConcertSyncTest::Replication::PropertyChain
 	 * Tests that UE::ConcertSyncCore::FReplicationPropertyFilter works correctly.
 	 * This uses the filter and serializes an object, doing a Memcmp to detect that exactly the expected properties are serialized.
 	 */
-	IMPLEMENT_SIMPLE_AUTOMATION_TEST(FSerializeWithPropertyFilterTest, "Concert.Replication.Data.SerializeWithPropertyFilter", EAutomationTestFlags::EditorContext | EAutomationTestFlags::EngineFilter)
+	IMPLEMENT_SIMPLE_AUTOMATION_TEST(FSerializeWithPropertyFilterTest, "Editor.Concert.Replication.Data.SerializeWithPropertyFilter", EAutomationTestFlags::EditorContext | EAutomationTestFlags::EngineFilter)
 	bool FSerializeWithPropertyFilterTest::RunTest(const FString& Parameters)
 	{
 		class FTestArchive : public FMemoryWriter
@@ -460,7 +460,7 @@ namespace UE::ConcertSyncTest::Replication::PropertyChain
 	 * Tests that UE::ConcertSyncCore::FReplicationPropertyFilter works correctly.
 	 * This uses the filter and serializes an object, doing a Memcmp to detect that exactly the expected properties are serialized.
 	 */
-	IMPLEMENT_SIMPLE_AUTOMATION_TEST(FOverlapPropertiesTest, "Concert.Replication.Data.OverlapProperties", EAutomationTestFlags::EditorContext | EAutomationTestFlags::EngineFilter)
+	IMPLEMENT_SIMPLE_AUTOMATION_TEST(FOverlapPropertiesTest, "Editor.Concert.Replication.Data.OverlapProperties", EAutomationTestFlags::EditorContext | EAutomationTestFlags::EngineFilter)
 	bool FOverlapPropertiesTest::RunTest(const FString& Parameters)
 	{
 		// 1. Define paths
diff --git a/Engine/Plugins/Developer/Concert/ConcertSync/ConcertSyncTest/Source/ConcertSyncTest/Private/Replication/ReplicationAuthorityTests.cpp b/Engine/Plugins/Developer/Concert/ConcertSync/ConcertSyncTest/Source/ConcertSyncTest/Private/Replication/ReplicationAuthorityTests.cpp
--- a/Engine/Plugins/Developer/Concert/ConcertSync/ConcertSyncTest/Source/ConcertSyncTest/Private/Replication/ReplicationAuthorityTests.cpp
+++ b/Engine/Plugins/Developer/Concert/ConcertSync/ConcertSyncTest/Source/ConcertSyncTest/Private/Replication/ReplicationAuthorityTests.cpp
@@ -18,7 +18,7 @@
 namespace UE::ConcertSyncTests::Replication::Authority
 {
 	/** Data sent from a client that does not have authority over objects is rejected. */
-	IMPLEMENT_CUSTOM_SIMPLE_AUTOMATION_TEST(FRejectUnauthorativeClientTest, FSendReceiveObjectTestBase, "Concert.Replication.Authority.RejectUnauthorativeClient", EAutomationTestFlags::EditorContext | EAutomationTestFlags::EngineFilter);
+	IMPLEMENT_CUSTOM_SIMPLE_AUTOMATION_TEST(FRejectUnauthorativeClientTest, FSendReceiveObjectTestBase, "Editor.Concert.Replication.Authority.RejectUnauthorativeClient", EAutomationTestFlags::EditorContext | EAutomationTestFlags::EngineFilter);
 	bool FRejectUnauthorativeClientTest::RunTest(const FString& Parameters)
 	{
 		// 1. Init
@@ -92,7 +92,7 @@ namespace UE::ConcertSyncTests::Replication::Authority
 	}
 
 	/** Makes sure only one client can take authority over the same object properties at the same time and that other clients continue being able to take authority after authority has been released.  */
-	IMPLEMENT_CUSTOM_SIMPLE_AUTOMATION_TEST(FAuthorityIsMutallyExclusiveTest, FSendReceiveGenericTestBase, "Concert.Replication.Authority.AuthorityIsMutallyExclusive", EAutomationTestFlags::EditorContext | EAutomationTestFlags::EngineFilter);
+	IMPLEMENT_CUSTOM_SIMPLE_AUTOMATION_TEST(FAuthorityIsMutallyExclusiveTest, FSendReceiveGenericTestBase, "Editor.Concert.Replication.Authority.AuthorityIsMutallyExclusive", EAutomationTestFlags::EditorContext | EAutomationTestFlags::EngineFilter);
 	bool FAuthorityIsMutallyExclusiveTest::RunTest(const FString& Parameters)
 	{
 		// 1. Init
@@ -172,7 +172,7 @@ namespace UE::ConcertSyncTests::Replication::Authority
 	}
 	
 	/** Clients cannot take authority over objects that they did not send */
-	IMPLEMENT_CUSTOM_SIMPLE_AUTOMATION_TEST(FCannotTakeAuthorityOverUnregisteredObjectsTest, FSendReceiveGenericTestBase, "Concert.Replication.Authority.CannotTakeAuthorityOverUnregisteredObjects", EAutomationTestFlags::EditorContext | EAutomationTestFlags::EngineFilter);
+	IMPLEMENT_CUSTOM_SIMPLE_AUTOMATION_TEST(FCannotTakeAuthorityOverUnregisteredObjectsTest, FSendReceiveGenericTestBase, "Editor.Concert.Replication.Authority.CannotTakeAuthorityOverUnregisteredObjects", EAutomationTestFlags::EditorContext | EAutomationTestFlags::EngineFilter);
 	bool FCannotTakeAuthorityOverUnregisteredObjectsTest::RunTest(const FString& Parameters)
 	{
 		// FSendReceiveGenericTestBase is used here because it registers no sending streams by default
@@ -214,7 +214,7 @@ namespace UE::ConcertSyncTests::Replication::Authority
 	}
 
 	/** When a client leaves the replication session, the client releases all of the authority as well.  */
-	IMPLEMENT_CUSTOM_SIMPLE_AUTOMATION_TEST(FClientLeavingReplicationSessionLosesAuthority, FSendReceiveGenericTestBase, "Concert.Replication.Authority.ClientLeavingReplicationSessionLosesAuthority", EAutomationTestFlags::EditorContext | EAutomationTestFlags::EngineFilter);
+	IMPLEMENT_CUSTOM_SIMPLE_AUTOMATION_TEST(FClientLeavingReplicationSessionLosesAuthority, FSendReceiveGenericTestBase, "Editor.Concert.Replication.Authority.ClientLeavingReplicationSessionLosesAuthority", EAutomationTestFlags::EditorContext | EAutomationTestFlags::EngineFilter);
 	bool FClientLeavingReplicationSessionLosesAuthority::RunTest(const FString& Parameters)
 	{
 		// 1. Init
@@ -260,7 +260,7 @@ namespace UE::ConcertSyncTests::Replication::Authority
 	}
 
 	/** When a client leaves the replication session, the client releases all of the authority as well.  */
-	IMPLEMENT_CUSTOM_SIMPLE_AUTOMATION_TEST(FClientLeavingConcertSessionLosesAuthority, FSendReceiveGenericTestBase, "Concert.Replication.Authority.ClientLeavingConcertSessionLosesAuthority", EAutomationTestFlags::EditorContext | EAutomationTestFlags::EngineFilter);
+	IMPLEMENT_CUSTOM_SIMPLE_AUTOMATION_TEST(FClientLeavingConcertSessionLosesAuthority, FSendReceiveGenericTestBase, "Editor.Concert.Replication.Authority.ClientLeavingConcertSessionLosesAuthority", EAutomationTestFlags::EditorContext | EAutomationTestFlags::EngineFilter);
 	bool FClientLeavingConcertSessionLosesAuthority::RunTest(const FString& Parameters)
 	{
 		// 1. Init
@@ -310,7 +310,7 @@ namespace UE::ConcertSyncTests::Replication::Authority
 	 * - Taking authority calls IConcertClientReplicationBridge::PushTrackedObjects
 	 * - Releasing authority calls IConcertClientReplicationBridge::ReleaseTrackedObjects
 	 */
-	IMPLEMENT_CUSTOM_SIMPLE_AUTOMATION_TEST(FChangingAuthorityUpdatesClientBridge, FSendReceiveObjectTestBase, "Concert.Replication.Authority.ChangingAuthorityUpdatesClientBridge", EAutomationTestFlags::EditorContext | EAutomationTestFlags::EngineFilter);
+	IMPLEMENT_CUSTOM_SIMPLE_AUTOMATION_TEST(FChangingAuthorityUpdatesClientBridge, FSendReceiveObjectTestBase, "Editor.Concert.Replication.Authority.ChangingAuthorityUpdatesClientBridge", EAutomationTestFlags::EditorContext | EAutomationTestFlags::EngineFilter);
 	bool FChangingAuthorityUpdatesClientBridge::RunTest(const FString& Parameters)
 	{
 		// 1. Init
@@ -352,7 +352,7 @@ namespace UE::ConcertSyncTests::Replication::Authority
 	/**
 	 * Tests that client updates its local cache of the server state when RequestAuthorityChange times out.
 	 */
-	IMPLEMENT_CUSTOM_SIMPLE_AUTOMATION_TEST(FChangingAuthorityTimeoutRetainsServerState, FSendReceiveObjectTestBase, "Concert.Replication.Authority.ChangeAuthorityTimeoutRetainsServerState", EAutomationTestFlags::EditorContext | EAutomationTestFlags::EngineFilter);
+	IMPLEMENT_CUSTOM_SIMPLE_AUTOMATION_TEST(FChangingAuthorityTimeoutRetainsServerState, FSendReceiveObjectTestBase, "Editor.Concert.Replication.Authority.ChangeAuthorityTimeoutRetainsServerState", EAutomationTestFlags::EditorContext | EAutomationTestFlags::EngineFilter);
 	bool FChangingAuthorityTimeoutRetainsServerState::RunTest(const FString& Parameters)
 	{
 		// 1. Init
@@ -431,7 +431,7 @@ namespace UE::ConcertSyncTests::Replication::Authority
 	}
 
 	/** Tests UE::ConcertSyncCore::Replication::AuthorityConflictUtils::EnumerateAuthorityConflicts. */
-	IMPLEMENT_CUSTOM_SIMPLE_AUTOMATION_TEST(FConflictEnumerationComponenTests, FSendReceiveObjectTestBase, "Concert.Replication.Authority.ConflictEnumerationComponenTests", EAutomationTestFlags::EditorContext | EAutomationTestFlags::EngineFilter);
+	IMPLEMENT_CUSTOM_SIMPLE_AUTOMATION_TEST(FConflictEnumerationComponenTests, FSendReceiveObjectTestBase, "Editor.Concert.Replication.Authority.ConflictEnumerationComponenTests", EAutomationTestFlags::EditorContext | EAutomationTestFlags::EngineFilter);
 	bool FConflictEnumerationComponenTests::RunTest(const FString& Parameters)
 	{
 		using namespace ConflictEnumerationComponent;
diff --git a/Engine/Plugins/Developer/Concert/ConcertSync/ConcertSyncTest/Source/ConcertSyncTest/Private/Replication/ReplicationHandshakeTests.cpp b/Engine/Plugins/Developer/Concert/ConcertSync/ConcertSyncTest/Source/ConcertSyncTest/Private/Replication/ReplicationHandshakeTests.cpp
--- a/Engine/Plugins/Developer/Concert/ConcertSync/ConcertSyncTest/Source/ConcertSyncTest/Private/Replication/ReplicationHandshakeTests.cpp
+++ b/Engine/Plugins/Developer/Concert/ConcertSync/ConcertSyncTest/Source/ConcertSyncTest/Private/Replication/ReplicationHandshakeTests.cpp
@@ -21,7 +21,7 @@ namespace UE::ConcertSyncTests::Replication::Handshake
 	/**
 	 * Tests the handshake for joining and leaving a replication session.
 	 */
-	IMPLEMENT_CUSTOM_SIMPLE_AUTOMATION_TEST(FJoinHandshakeTest, FConcertClientServerCommunicationTest, "Concert.Replication.Handshake", EAutomationTestFlags::EditorContext | EAutomationTestFlags::EngineFilter);
+	IMPLEMENT_CUSTOM_SIMPLE_AUTOMATION_TEST(FJoinHandshakeTest, FConcertClientServerCommunicationTest, "Editor.Concert.Replication.Handshake", EAutomationTestFlags::EditorContext | EAutomationTestFlags::EngineFilter);
 	bool FJoinHandshakeTest::RunTest(const FString& Parameters)
 	{
 		using namespace ConcertSyncClient::TestInterface;
diff --git a/Engine/Plugins/Developer/Concert/ConcertSync/ConcertSyncTest/Source/ConcertSyncTest/Private/Replication/ReplicationSendReceiveTests.cpp b/Engine/Plugins/Developer/Concert/ConcertSync/ConcertSyncTest/Source/ConcertSyncTest/Private/Replication/ReplicationSendReceiveTests.cpp
--- a/Engine/Plugins/Developer/Concert/ConcertSync/ConcertSyncTest/Source/ConcertSyncTest/Private/Replication/ReplicationSendReceiveTests.cpp
+++ b/Engine/Plugins/Developer/Concert/ConcertSync/ConcertSyncTest/Source/ConcertSyncTest/Private/Replication/ReplicationSendReceiveTests.cpp
@@ -91,7 +91,7 @@ namespace UE::ConcertSyncTests::Replication::SendReceiveFlow
 	/**
 	 * Tests replicating data from sender client > server > receiver client.
 	 */
-	IMPLEMENT_CUSTOM_SIMPLE_AUTOMATION_TEST(FSendReceiveFlowTests, FSendReceiveObjectTestBase, "Concert.Replication.SendReceive.SingleStream", EAutomationTestFlags::EditorContext | EAutomationTestFlags::EngineFilter);
+	IMPLEMENT_CUSTOM_SIMPLE_AUTOMATION_TEST(FSendReceiveFlowTests, FSendReceiveObjectTestBase, "Editor.Concert.Replication.SendReceive.SingleStream", EAutomationTestFlags::EditorContext | EAutomationTestFlags::EngineFilter);
 	bool FSendReceiveFlowTests::RunTest(const FString& Parameters)
 	{
 		SharedSetupSet(*this);
@@ -140,7 +140,7 @@ namespace UE::ConcertSyncTests::Replication::SendReceiveFlow
 	};
 
 	/** Test which still sends TestObject but does so with two separate streams: one for the float and the other for the vector property. */
-	IMPLEMENT_CUSTOM_SIMPLE_AUTOMATION_TEST(FSendReceiveMultiStreamSameObjectTests, FSplitSendReceiveObjectTest, "Concert.Replication.SendReceive.MultipleStreamsForSameObject", EAutomationTestFlags::EditorContext | EAutomationTestFlags::EngineFilter);
+	IMPLEMENT_CUSTOM_SIMPLE_AUTOMATION_TEST(FSendReceiveMultiStreamSameObjectTests, FSplitSendReceiveObjectTest, "Editor.Concert.Replication.SendReceive.MultipleStreamsForSameObject", EAutomationTestFlags::EditorContext | EAutomationTestFlags::EngineFilter);
 	bool FSendReceiveMultiStreamSameObjectTests::RunTest(const FString& Parameters)
 	{
 		SharedSetupSet(*this);
@@ -149,7 +149,7 @@ namespace UE::ConcertSyncTests::Replication::SendReceiveFlow
 	}
 
 	/** Tests that changing a stream while replication is in progress actually changes the properties being replicated. */
-	IMPLEMENT_CUSTOM_SIMPLE_AUTOMATION_TEST(FUpdateStreamWhileSendingUpdatesSentProperties, FSendReceiveObjectTestBase, "Concert.Replication.SendReceive.UpdateStreamWhileSendingUpdatesSentProperties", EAutomationTestFlags::EditorContext | EAutomationTestFlags::EngineFilter);
+	IMPLEMENT_CUSTOM_SIMPLE_AUTOMATION_TEST(FUpdateStreamWhileSendingUpdatesSentProperties, FSendReceiveObjectTestBase, "Editor.Concert.Replication.SendReceive.UpdateStreamWhileSendingUpdatesSentProperties", EAutomationTestFlags::EditorContext | EAutomationTestFlags::EngineFilter);
 	bool FUpdateStreamWhileSendingUpdatesSentProperties::RunTest(const FString& Parameters)
 	{
 		// 1. Send all properties with the stream from handshake
diff --git a/Engine/Plugins/Developer/Concert/ConcertSync/ConcertSyncTest/Source/ConcertSyncTest/Private/Replication/ReplicationStreamRequestTests.cpp b/Engine/Plugins/Developer/Concert/ConcertSync/ConcertSyncTest/Source/ConcertSyncTest/Private/Replication/ReplicationStreamRequestTests.cpp
--- a/Engine/Plugins/Developer/Concert/ConcertSync/ConcertSyncTest/Source/ConcertSyncTest/Private/Replication/ReplicationStreamRequestTests.cpp
+++ b/Engine/Plugins/Developer/Concert/ConcertSync/ConcertSyncTest/Source/ConcertSyncTest/Private/Replication/ReplicationStreamRequestTests.cpp
@@ -20,7 +20,7 @@ namespace UE::ConcertSyncTests::Replication::Stream
 	 * - that the authority updates,
 	 * - the EConcertQueryClientStreamFlags flags
 	 */
-	IMPLEMENT_CUSTOM_SIMPLE_AUTOMATION_TEST(FQueryOtherClientStreams, FSendReceiveObjectTestBase, "Concert.Replication.Stream.QueryClientStream", EAutomationTestFlags::EditorContext | EAutomationTestFlags::EngineFilter);
+	IMPLEMENT_CUSTOM_SIMPLE_AUTOMATION_TEST(FQueryOtherClientStreams, FSendReceiveObjectTestBase, "Editor.Concert.Replication.Stream.QueryClientStream", EAutomationTestFlags::EditorContext | EAutomationTestFlags::EngineFilter);
 	bool FQueryOtherClientStreams::RunTest(const FString& Parameters)
 	{
 		// 1. Init
@@ -276,7 +276,7 @@ namespace UE::ConcertSyncTests::Replication::Stream
 	};
 	
 	/** Simple case of FConcertChangeStream_Request where the requesting client has authority and thus cannot generate any authority conflicts. */
-	IMPLEMENT_CUSTOM_SIMPLE_AUTOMATION_TEST(FSimpleChangeStream, FChangeStreamsTestBase, "Concert.Replication.Stream.SimpleChangeStream", EAutomationTestFlags::EditorContext | EAutomationTestFlags::EngineFilter);
+	IMPLEMENT_CUSTOM_SIMPLE_AUTOMATION_TEST(FSimpleChangeStream, FChangeStreamsTestBase, "Editor.Concert.Replication.Stream.SimpleChangeStream", EAutomationTestFlags::EditorContext | EAutomationTestFlags::EngineFilter);
 	bool FSimpleChangeStream::RunTest(const FString& Parameters)
 	{
 		using namespace ConcertSyncClient::Replication;
@@ -314,7 +314,7 @@ namespace UE::ConcertSyncTests::Replication::Stream
 	}
 
 	/** Case of FConcertChangeStream_Request in which requester has authority over an object and creates two streams which have authority over overlapping properties. This is allowed. */
-	IMPLEMENT_CUSTOM_SIMPLE_AUTOMATION_TEST(FChangeStreamWithOverlappingProperties, FChangeStreamsTestBase, "Concert.Replication.Stream.ChangeStreamWithOverlappingProperties", EAutomationTestFlags::EditorContext | EAutomationTestFlags::EngineFilter);
+	IMPLEMENT_CUSTOM_SIMPLE_AUTOMATION_TEST(FChangeStreamWithOverlappingProperties, FChangeStreamsTestBase, "Editor.Concert.Replication.Stream.ChangeStreamWithOverlappingProperties", EAutomationTestFlags::EditorContext | EAutomationTestFlags::EngineFilter);
 	bool FChangeStreamWithOverlappingProperties::RunTest(const FString& Parameters)
 	{
 		using namespace ConcertSyncClient::Replication;
@@ -361,7 +361,7 @@ namespace UE::ConcertSyncTests::Replication::Stream
 	 * The requester and client have authority over separate properties on the same object.
 	 * The request would cause the requester's existing authority to overlap with the other client, which is rejected.
 	 */
-	IMPLEMENT_CUSTOM_SIMPLE_AUTOMATION_TEST(FRejectChangeStreamWithConflictingAuthority, FChangeStreamsTestBase, "Concert.Replication.Stream.RejectChangeStreamWithConflictingAuthority", EAutomationTestFlags::EditorContext | EAutomationTestFlags::EngineFilter);
+	IMPLEMENT_CUSTOM_SIMPLE_AUTOMATION_TEST(FRejectChangeStreamWithConflictingAuthority, FChangeStreamsTestBase, "Editor.Concert.Replication.Stream.RejectChangeStreamWithConflictingAuthority", EAutomationTestFlags::EditorContext | EAutomationTestFlags::EngineFilter);
 	bool FRejectChangeStreamWithConflictingAuthority::RunTest(const FString& Parameters)
 	{
 		using namespace ConcertSyncClient::Replication;
@@ -447,7 +447,7 @@ namespace UE::ConcertSyncTests::Replication::Stream
 	}
 
 	/** Tests that no changes are made if a sub-step of FConcertChangeStream_Request causes an error. */
-	IMPLEMENT_CUSTOM_SIMPLE_AUTOMATION_TEST(FChangeStreamIsAtmoic, FChangeStreamsTestBase, "Concert.Replication.Stream.ChangeStreamIsAtmoic", EAutomationTestFlags::EditorContext | EAutomationTestFlags::EngineFilter);
+	IMPLEMENT_CUSTOM_SIMPLE_AUTOMATION_TEST(FChangeStreamIsAtmoic, FChangeStreamsTestBase, "Editor.Concert.Replication.Stream.ChangeStreamIsAtmoic", EAutomationTestFlags::EditorContext | EAutomationTestFlags::EngineFilter);
 	bool FChangeStreamIsAtmoic::RunTest(const FString& Parameters)
 	{
 		using namespace ConcertSyncClient::Replication;
@@ -488,7 +488,7 @@ namespace UE::ConcertSyncTests::Replication::Stream
 	}
 
 	/** Removing an object from a stream over which the requester has authority, clears the authority the client had. Tests that another client can now claim authority over the previously authored properties. */
-	IMPLEMENT_CUSTOM_SIMPLE_AUTOMATION_TEST(FRemovingObjectClearsAuthority, FChangeStreamsTestBase, "Concert.Replication.Stream.RemovingObjectClearsAuthority", EAutomationTestFlags::EditorContext | EAutomationTestFlags::EngineFilter);
+	IMPLEMENT_CUSTOM_SIMPLE_AUTOMATION_TEST(FRemovingObjectClearsAuthority, FChangeStreamsTestBase, "Editor.Concert.Replication.Stream.RemovingObjectClearsAuthority", EAutomationTestFlags::EditorContext | EAutomationTestFlags::EngineFilter);
 	bool FRemovingObjectClearsAuthority::RunTest(const FString& Parameters)
 	{
 		using namespace ConcertSyncClient::Replication;
@@ -528,7 +528,7 @@ namespace UE::ConcertSyncTests::Replication::Stream
 	/**
      * Tests that client updates its local cache of the server state when RequestAuthorityChange times out.
      */
-    IMPLEMENT_CUSTOM_SIMPLE_AUTOMATION_TEST(FChangingStreamTimeoutRetainsServerState, FChangeStreamsTestBase, "Concert.Replication.Stream.ChangeStreamTimeoutRetainsServerState", EAutomationTestFlags::EditorContext | EAutomationTestFlags::EngineFilter);
+    IMPLEMENT_CUSTOM_SIMPLE_AUTOMATION_TEST(FChangingStreamTimeoutRetainsServerState, FChangeStreamsTestBase, "Editor.Concert.Replication.Stream.ChangeStreamTimeoutRetainsServerState", EAutomationTestFlags::EditorContext | EAutomationTestFlags::EngineFilter);
     bool FChangingStreamTimeoutRetainsServerState::RunTest(const FString& Parameters)
     {
     	// 1. Start up client with float property stream
