Add object replication frequency.
- FConcertReplication_ChangeStream has been extended with a way to configure replication frequency on server
- Client and server send objects at their configured frequency
- Request and frequency replication are unit tested

diff --git a/Engine/Plugins/Developer/Concert/ConcertSync/ConcertSyncClient/Source/ConcertSyncClient/Private/Replication/Manager/ReplicationManagerState_Connected.cpp b/Engine/Plugins/Developer/Concert/ConcertSync/ConcertSyncClient/Source/ConcertSyncClient/Private/Replication/Manager/ReplicationManagerState_Connected.cpp
--- a/Engine/Plugins/Developer/Concert/ConcertSync/ConcertSyncClient/Source/ConcertSyncClient/Private/Replication/Manager/ReplicationManagerState_Connected.cpp
+++ b/Engine/Plugins/Developer/Concert/ConcertSync/ConcertSyncClient/Source/ConcertSyncClient/Private/Replication/Manager/ReplicationManagerState_Connected.cpp
@@ -61,9 +61,13 @@ namespace UE::ConcertSyncClient::Replication
 			FClientReplicationDataCollector::FGetClientStreams::CreateLambda([this]()
 			{
 				return &RegisteredStreams;
-			}))
+			}),
+			LiveSession->GetSessionClientEndpointId()
+			))
+		, Sender(
+			ConcertSyncCore::FGetObjectFrequencySettings::CreateRaw(this, &FReplicationManagerState_Connected::GetObjectFrequencySettings),
+			LiveSession->GetSessionServerEndpointId(), LiveSession, ReplicationDataSource
 			)
-		, Sender(MakeShared<ConcertSyncCore::FObjectReplicationSender>(LiveSession->GetSessionServerEndpointId(), LiveSession, ReplicationDataSource))
 		, ReceivedDataCache(MakeShared<ConcertSyncCore::FObjectReplicationCache>(ReplicationFormat))
 		, Receiver(MakeShared<ConcertSyncCore::FObjectReplicationReceiver>(LiveSession, ReceivedDataCache))
 		, ReceivedReplicationQueuer(FClientReplicationDataQueuer::Make(ReplicationBridge, ReceivedDataCache))
@@ -142,7 +146,7 @@ namespace UE::ConcertSyncClient::Replication
 			return MakeFulfilledPromise<FConcertReplication_QueryReplicationInfo_Response>(FConcertReplication_QueryReplicationInfo_Response{ EReplicationResponseErrorCode::Timeout }).GetFuture();
 		}
 		
-		if (EnumHasAllFlags(Args.QueryFlags, EConcertQueryClientStreamFlags::SkipAuthority | EConcertQueryClientStreamFlags::SkipStreamInfo))
+		if (EnumHasAllFlags(Args.QueryFlags, EConcertQueryClientStreamFlags::SkipAuthority | EConcertQueryClientStreamFlags::SkipStreamInfo | EConcertQueryClientStreamFlags::SkipFrequency ))
 		{
 			UE_LOG(LogConcert, Warning, TEXT("Request QueryClientInfo is pointless because SkipAuthority and SkipStreamInfo are both set. Returning immediately..."));
 			return MakeFulfilledPromise<FConcertReplication_QueryReplicationInfo_Response>().GetFuture();
@@ -218,34 +222,10 @@ namespace UE::ConcertSyncClient::Replication
 
 	void FReplicationManagerState_Connected::Tick(IConcertClientSession& Session, float DeltaTime)
 	{
-		// TODO DP: Set this up in a config file
-		constexpr double TimeBudget = 1.0 / 60.0;
-		double TimeLeft = TimeBudget;
-
-		auto NextIndex = [this](int32 Current){ return AlternatingTickTasks.IsValidIndex(Current + 1) ? Current + 1 : 0; };
-		int32 CurrentIndex = NextTickTaskIndex;
-		const double StartTime = FPlatformTime::Seconds();
-		do
-		{
-			FTickTask Task = AlternatingTickTasks[CurrentIndex]; 
-			Invoke(Task, this, TimeLeft);
-
-			const double TotalElapsedTime = FPlatformTime::Seconds() - StartTime;
-			TimeLeft = TimeBudget - TotalElapsedTime;
-			CurrentIndex = NextIndex(CurrentIndex);
-		}
-		while (TimeLeft > 0.0 && CurrentIndex != NextTickTaskIndex);
-		NextTickTaskIndex = CurrentIndex;
-	}
-
-	void FReplicationManagerState_Connected::TickSender(float TimeBudget)
-	{
-		Sender->ProcessObjects(TimeBudget);
-	}
-
-	void FReplicationManagerState_Connected::TickReceiver(float TimeBudget)
-	{
-		ReplicationApplier->ProcessObjects(TimeBudget);
+		// TODO UE-190714: We should set a time budget for the client so ticking does not cause frame spikes
+		const ConcertSyncCore::FProcessObjectsParams Params { DeltaTime };
+		Sender.ProcessObjects(Params);
+		ReplicationApplier->ProcessObjects(Params);
 	}
 
 	void FReplicationManagerState_Connected::UpdateReplicatedObjectsAfterStreamChange(const FConcertReplication_ChangeStream_Request& Request, const FConcertReplication_ChangeStream_Response& Response)
@@ -360,5 +340,21 @@ namespace UE::ConcertSyncClient::Replication
 			ReplicationDataSource->AddReplicatedObjectStreams(ReleaseAuthority.Key, ReleaseAuthority.Value.StreamIds);
 		}
 	}
+
+	FConcertObjectReplicationSettings FReplicationManagerState_Connected::GetObjectFrequencySettings(const FReplicatedObjectId& Object) const
+	{
+		const FReplicationStreamDescription* Stream = RegisteredStreams.FindByPredicate([&Object](const FReplicationStreamDescription& Description)
+		{
+			return Description.BaseDescription.Identifier == Object.StreamId;
+		});
+		
+		if (!ensureMsgf(Stream, TEXT("Caller of GetObjectFrequencySettings is trying to send an object that is not registered with the client")))
+		{
+			UE_LOG(LogConcert, Warning, TEXT("Requested frequency settings for unknown stream %s and object %s"), *Object.StreamId.ToString(), *Object.Object.ToString());
+			return {};
+		}
+		
+		return Stream->BaseDescription.FrequencySettings.GetSettingsFor(Object.Object);
+	}
 }
 
diff --git a/Engine/Plugins/Developer/Concert/ConcertSync/ConcertSyncClient/Source/ConcertSyncClient/Private/Replication/Manager/ReplicationManagerState_Connected.h b/Engine/Plugins/Developer/Concert/ConcertSync/ConcertSyncClient/Source/ConcertSyncClient/Private/Replication/Manager/ReplicationManagerState_Connected.h
--- a/Engine/Plugins/Developer/Concert/ConcertSync/ConcertSyncClient/Source/ConcertSyncClient/Private/Replication/Manager/ReplicationManagerState_Connected.h
+++ b/Engine/Plugins/Developer/Concert/ConcertSync/ConcertSyncClient/Source/ConcertSyncClient/Private/Replication/Manager/ReplicationManagerState_Connected.h
@@ -5,19 +5,16 @@
 #include "ReplicationManagerState.h"
 #include "Replication/Processing/ClientReplicationDataCollector.h"
 #include "Replication/Processing/ClientReplicationDataQueuer.h"
+#include "Replication/Processing/ObjectReplicationSender.h"
+#include "Replication/Processing/Proxy/ObjectProcessorProxy_Frequency.h"
 
-namespace UE::ConcertSyncCore
-{
-	class FObjectReplicationReceiver;
-}
+class IConcertClientSession;
 
 namespace UE::ConcertSyncCore
 {
-	class FObjectReplicationSender;
+	class FObjectReplicationReceiver;
 }
 
-class IConcertClientSession;
-
 namespace UE::ConcertSyncClient::Replication
 {
 	class FObjectReplicationApplierProcessor;
@@ -70,8 +67,11 @@ namespace UE::ConcertSyncClient::Replication
 		// Sending
 		/** Used as source of replication data. */
 		const TSharedRef<FClientReplicationDataCollector> ReplicationDataSource;
+		
+		/** Sends to remote endpoint and makes sure the objects are replicated at the specified frequency settings. */
+		using FDataRelayThrottledByFrequency = ConcertSyncCore::TObjectProcessorProxy_Frequency<ConcertSyncCore::FObjectReplicationSender>;
 		/** Sends data collected by ReplicationDataSource to the server. */
-		const TSharedRef<ConcertSyncCore::FObjectReplicationSender> Sender;
+		FDataRelayThrottledByFrequency Sender;
 
 		// Receiving
 		/** Stores data received by Receiver until it is consumed by ReceivedReplicationQueuer. */
@@ -83,20 +83,6 @@ namespace UE::ConcertSyncClient::Replication
 		/** Processes data from ReceivedReplicationQueuer once we tick. */
 		const TSharedRef<FObjectReplicationApplierProcessor> ReplicationApplier;
 
-		using FTickTask = void(FReplicationManagerState_Connected::*)(float TimeBudget);
-		/**
-		 * We have two tasks that need to be performed each tick: Collecting data to send and applying received data.
-		 * In order that no task starves the other, we alternate which one we start with.
-		 *
-		 * Example assuming time budget of 0.1s (unrealistic time budget):
-		 * Tick 1: TickSender takes 0.8s and finishes. TickReceiver has 0.2s but cannot finish all work.
-		 * Tick 2: TickReceiver finishes work taking 0.7s. TickSender gets the remaining 0.3s.
-		 * Tick 3. TickSender starts again followed by TickSender ...
-		 */
-		const TArray<FTickTask, TInlineAllocator<2>> AlternatingTickTasks { &FReplicationManagerState_Connected::TickSender, &FReplicationManagerState_Connected::TickReceiver };
-		/** The first index to process next tick. */
-		int32 NextTickTaskIndex = 0;
-
 		//~ Begin FReplicationManagerState Interface
 		virtual void OnEnterState() override;
 		//~ End FReplicationManagerState Interface
@@ -112,11 +98,6 @@ namespace UE::ConcertSyncClient::Replication
 		 * It is configured in the project settings TODO: Add config
 		 */
 		void Tick(IConcertClientSession& Session, float DeltaTime);
-
-		/** Collects and sends data to the server. */
-		void TickSender(float TimeBudget);
-		/** Processes received data and serializes UObjects. */
-		void TickReceiver(float TimeBudget);
 		
 		/** Updates replicated objects affected by the change request. */
 		void UpdateReplicatedObjectsAfterStreamChange(const FConcertReplication_ChangeStream_Request& Request, const FConcertReplication_ChangeStream_Response& Response);
@@ -132,5 +113,8 @@ namespace UE::ConcertSyncClient::Replication
 		void UpdateReplicatedObjectsAfterAuthorityChange(FConcertReplication_ChangeAuthority_Request&& Request, const FConcertReplication_ChangeAuthority_Response& Response) const;
 		void HandleReleasingReplicatedObjects(const FConcertReplication_ChangeAuthority_Request& Request) const;
 		void RevertReleasingReplicatedObjects(const FConcertReplication_ChangeAuthority_Request& Request) const;
+		
+		/** Callback to Sender for obtaining an object's frequency settings. */
+		FConcertObjectReplicationSettings GetObjectFrequencySettings(const FReplicatedObjectId& Object) const;
 	};
 }
diff --git a/Engine/Plugins/Developer/Concert/ConcertSync/ConcertSyncClient/Source/ConcertSyncClient/Private/Replication/Processing/ClientReplicationDataCollector.cpp b/Engine/Plugins/Developer/Concert/ConcertSync/ConcertSyncClient/Source/ConcertSyncClient/Private/Replication/Processing/ClientReplicationDataCollector.cpp
--- a/Engine/Plugins/Developer/Concert/ConcertSync/ConcertSyncClient/Source/ConcertSyncClient/Private/Replication/Processing/ClientReplicationDataCollector.cpp
+++ b/Engine/Plugins/Developer/Concert/ConcertSync/ConcertSyncClient/Source/ConcertSyncClient/Private/Replication/Processing/ClientReplicationDataCollector.cpp
@@ -16,11 +16,13 @@ namespace UE::ConcertSyncClient::Replication
 	FClientReplicationDataCollector::FClientReplicationDataCollector(
 		IConcertClientReplicationBridge* InReplicationBridge,
 		TSharedRef<ConcertSyncCore::IObjectReplicationFormat> InReplicationFormat,
-		FGetClientStreams InGetStreamsDelegate
+		FGetClientStreams InGetStreamsDelegate,
+		const FGuid& InClientId
 		)
 		: Bridge(InReplicationBridge)
 		, ReplicationFormat(MoveTemp(InReplicationFormat))
 		, GetStreamsDelegate(MoveTemp(InGetStreamsDelegate))
+		, ClientId(InClientId)
 	{
 		check(GetStreamsDelegate.IsBound());
 		Bridge->OnObjectDiscovered().AddRaw(this, &FClientReplicationDataCollector::StartTrackingObject);
@@ -153,7 +155,7 @@ namespace UE::ConcertSyncClient::Replication
 		}
 	}
 
-	void FClientReplicationDataCollector::ForEachPendingObject(TFunctionRef<void(const FObjectInStreamID&)> ProcessItemFunc) const
+	void FClientReplicationDataCollector::ForEachPendingObject(TFunctionRef<void(const FReplicatedObjectId&)> ProcessItemFunc) const
 	{
 		for (const TPair<FSoftObjectPath, TArray<FObjectInfo>>& Pair : ObjectsToReplicate)
 		{
@@ -161,14 +163,17 @@ namespace UE::ConcertSyncClient::Replication
 			{
 				if (ObjectInfo.ObjectCache.IsValid())
 				{
-					ProcessItemFunc({ ObjectInfo.StreamId, ObjectInfo.ObjectCache.Get() });
+					ProcessItemFunc(FReplicatedObjectId {
+						FObjectInStreamID{ ObjectInfo.StreamId, ObjectInfo.ObjectCache.Get() },
+						ClientId
+					});
 				}
 			}
 		}
 	}
 
 	bool FClientReplicationDataCollector::ExtractReplicationDataForObject(
-		const FObjectInStreamID& ObjectToProcess,
+		const FReplicatedObjectId& ObjectToProcess,
 		TFunctionRef<void(const FConcertSessionSerializedPayload& Payload)> ProcessCopyable,
 		TFunctionRef<void(FConcertSessionSerializedPayload&& Payload)> ProcessMoveable
 		)
diff --git a/Engine/Plugins/Developer/Concert/ConcertSync/ConcertSyncClient/Source/ConcertSyncClient/Private/Replication/Processing/ClientReplicationDataCollector.h b/Engine/Plugins/Developer/Concert/ConcertSync/ConcertSyncClient/Source/ConcertSyncClient/Private/Replication/Processing/ClientReplicationDataCollector.h
--- a/Engine/Plugins/Developer/Concert/ConcertSync/ConcertSyncClient/Source/ConcertSyncClient/Private/Replication/Processing/ClientReplicationDataCollector.h
+++ b/Engine/Plugins/Developer/Concert/ConcertSync/ConcertSyncClient/Source/ConcertSyncClient/Private/Replication/Processing/ClientReplicationDataCollector.h
@@ -35,10 +35,18 @@ namespace UE::ConcertSyncClient::Replication
 	public:
 
 		DECLARE_DELEGATE_RetVal(const TArray<FReplicationStreamDescription>*, FGetClientStreams);
+
+		/**
+		 * @param InReplicationBridge 
+		 * @param InReplicationFormat 
+		 * @param InGetStreamsDelegate 
+		 * @param ClientId 
+		 */
 		FClientReplicationDataCollector(
 			IConcertClientReplicationBridge* InReplicationBridge,
 			TSharedRef<ConcertSyncCore::IObjectReplicationFormat> InReplicationFormat,
-			FGetClientStreams InGetStreamsDelegate
+			FGetClientStreams InGetStreamsDelegate,
+			const FGuid& InClientId
 			);
 		virtual ~FClientReplicationDataCollector() override;
 		
@@ -67,9 +75,9 @@ namespace UE::ConcertSyncClient::Replication
 		void AppendOwningStreamsForObject(const FSoftObjectPath& ObjectPath, TSet<FGuid>& Paths) const;
 
 		//~ Begin IReplicationDataSource Interface
-		virtual void ForEachPendingObject(TFunctionRef<void(const FObjectInStreamID&)> ProcessItemFunc) const override;
+		virtual void ForEachPendingObject(TFunctionRef<void(const FReplicatedObjectId&)> ProcessItemFunc) const override;
 		virtual int32 NumObjects() const override { return NumTrackedObjects; }
-		virtual bool ExtractReplicationDataForObject(const FObjectInStreamID& Object, TFunctionRef<void(const FConcertSessionSerializedPayload& Payload)> ProcessCopyable, TFunctionRef<void(FConcertSessionSerializedPayload&& Payload)> ProcessMoveable) override;
+		virtual bool ExtractReplicationDataForObject(const FReplicatedObjectId& Object, TFunctionRef<void(const FConcertSessionSerializedPayload& Payload)> ProcessCopyable, TFunctionRef<void(FConcertSessionSerializedPayload&& Payload)> ProcessMoveable) override;
 		//~ End IReplicationDataSource Interface
 
 	private:
@@ -82,6 +90,9 @@ namespace UE::ConcertSyncClient::Replication
 		/** Gets the stream of the managed client. */
 		const FGetClientStreams GetStreamsDelegate;
 
+		/** Endpoint ID of the client */
+		const FGuid ClientId;
+
 		struct FObjectInfo
 		{
 			/** The replication stream producing this object's data */
diff --git a/Engine/Plugins/Developer/Concert/ConcertSync/ConcertSyncCore/Source/ConcertSyncCore/Private/Replication/ChangeStreamSharedUtils.cpp b/Engine/Plugins/Developer/Concert/ConcertSync/ConcertSyncCore/Source/ConcertSyncCore/Private/Replication/ChangeStreamSharedUtils.cpp
--- a/Engine/Plugins/Developer/Concert/ConcertSync/ConcertSyncCore/Source/ConcertSyncCore/Private/Replication/ChangeStreamSharedUtils.cpp
+++ b/Engine/Plugins/Developer/Concert/ConcertSync/ConcertSyncCore/Source/ConcertSyncCore/Private/Replication/ChangeStreamSharedUtils.cpp
@@ -48,6 +48,7 @@ namespace UE::ConcertSyncCore::Replication::ChangeStreamUtils
 		{
 			FSharedReplicationStreamDescription& BaseDescription = StreamIt->BaseDescription;
 			TMap<FSoftObjectPath, FReplicatedObjectInfo>& ReplicationMap = BaseDescription.ReplicationMap.ReplicatedObjects;
+			FConcertStreamFrequencySettings& FrequencySettings = BaseDescription.FrequencySettings;
 			const FGuid& StreamId = BaseDescription.Identifier;
 			
 			for (const FObjectInStreamID& ObjectToRemove : Request.ObjectsToRemove)
@@ -55,8 +56,14 @@ namespace UE::ConcertSyncCore::Replication::ChangeStreamUtils
 				if (ObjectToRemove.StreamId == StreamId)
 				{
 					ReplicationMap.Remove(ObjectToRemove.Object);
+					FrequencySettings.ObjectOverrides.Remove(ObjectToRemove.Object);
 				}
 			}
+
+			if (const FConcertReplication_ChangeStream_Frequency* FrequencyChange = Request.FrequencyChanges.Find(StreamId))
+			{
+				ApplyValidatedFrequencyChanges(*FrequencyChange, FrequencySettings);
+			}
 			
 			if (Request.StreamsToRemove.Contains(StreamId)
 				// By contract, ObjectsToRemove will remove the stream if the stream ends up being empty
@@ -104,6 +111,128 @@ namespace UE::ConcertSyncCore::Replication::ChangeStreamUtils
 		}
 	}
 	
+	void ApplyValidatedFrequencyChanges(
+		const FConcertReplication_ChangeStream_Frequency& Request,
+		FConcertStreamFrequencySettings& SettingsToModify
+		)
+	{
+		if (EnumHasAnyFlags(Request.Flags, EConcertReplicationChangeFrequencyFlags::SetDefaults)
+			&& ensure(Request.NewDefaults.IsValid()))
+		{
+			SettingsToModify.Defaults = Request.NewDefaults;
+		}
+
+		if (!Request.OverridesToPut.IsEmpty())
+		{
+			SettingsToModify.ObjectOverrides = Request.OverridesToPut;
+		}
+
+		for (const FSoftObjectPath& ToRemove : Request.OverridesToRemove)
+		{
+			SettingsToModify.ObjectOverrides.Remove(ToRemove);
+		}
+		
+		for (const TPair<FSoftObjectPath, FConcertObjectReplicationSettings>& ToAdd : Request.OverridesToAdd)
+		{
+			if (ensure(ToAdd.Value.IsValid()))
+			{
+				SettingsToModify.ObjectOverrides.Add(ToAdd.Key, ToAdd.Value);
+			}
+		}
+	}
+
+#define ADD_ERROR(Condition, Callback) if (Condition) { Callback; }
+	namespace Private
+	{
+		static void MarkAsMissingStream(const FGuid& StreamId, const TMap<FSoftObjectPath, FConcertObjectReplicationSettings>& Added, FConcertReplication_ChangeStream_FrequencyResponse& Errors)
+		{
+			for (const TPair<FSoftObjectPath, FConcertObjectReplicationSettings>& Pair : Added)
+			{
+				Errors.OverrideFailures.Add({ StreamId, Pair.Key }, EConcertChangeObjectFrequencyErrorCode::NotRegistered);
+			}
+		}
+
+		static bool HasOrIsAddingProperties(const FConcertReplication_ChangeStream_Request& Request, const FSharedReplicationStreamDescription& Stream, const FSoftObjectPath& ObjectPath)
+		{
+			const FObjectReplicationMap& ReplicationMap = Stream.ReplicationMap;
+			return ReplicationMap.HasProperties(ObjectPath)
+				// It might be that ObjectsToPut is invalid.
+				// We do not need to check that here because the underlying FConcertReplication_ChangeStream_Request will be checked separately and would fail then anyways.
+				|| Request.ObjectsToPut.Contains({ Stream.Identifier, ObjectPath });
+		}
+		
+		static bool ValidateAddedFrequencies(
+			const FConcertReplication_ChangeStream_Request& FullRequest,
+			const TMap<FSoftObjectPath, FConcertObjectReplicationSettings>& Added,
+			const FSharedReplicationStreamDescription& Stream,
+			FConcertReplication_ChangeStream_FrequencyResponse* OptionalErrors = nullptr)
+		{
+			bool bAnyError = false;
+			for (const TPair<FSoftObjectPath, FConcertObjectReplicationSettings>& Pair : Added)
+			{
+				if (!HasOrIsAddingProperties(FullRequest, Stream, Pair.Key))
+				{
+					bAnyError = true;
+					ADD_ERROR(OptionalErrors, OptionalErrors->OverrideFailures.Add({ Stream.Identifier, Pair.Key }, EConcertChangeObjectFrequencyErrorCode::NotRegistered));
+				}
+				else if (!Pair.Value.IsValid())
+				{
+					ADD_ERROR(OptionalErrors, OptionalErrors->OverrideFailures.Add({ Stream.Identifier, Pair.Key }, EConcertChangeObjectFrequencyErrorCode::InvalidReplicationRate));
+				}
+			}
+			
+			return bAnyError;
+		}
+	}
+
+	bool ValidateFrequencyChanges(
+		const FConcertReplication_ChangeStream_Request& Request,
+		const TArray<FReplicationStreamDescription>& Streams,
+		FConcertReplication_ChangeStream_FrequencyResponse* OptionalErrors
+		)
+	{
+		using namespace Private;
+		bool bAnyErrors = false;
+		
+		for (const TPair<FGuid, FConcertReplication_ChangeStream_Frequency>& SubRequestPair : Request.FrequencyChanges)
+		{
+			const FGuid& TargetStreamId = SubRequestPair.Key;
+			const FConcertReplication_ChangeStream_Frequency& SubRequest = SubRequestPair.Value;
+			
+			const bool bIsModifyingDefaults = EnumHasAnyFlags(SubRequest.Flags, EConcertReplicationChangeFrequencyFlags::SetDefaults);
+			const FReplicationStreamDescription* StreamDescription = Streams.FindByPredicate([&TargetStreamId](const FReplicationStreamDescription& Description)
+			{
+				return Description.BaseDescription.Identifier == TargetStreamId;
+			});
+
+			// Referenced stream does not exist
+			if (!StreamDescription)
+			{
+				// TODO UE-201167: Once StreamsToAdd no longer specifies objects directly, we need to handle that case here
+				
+				bAnyErrors = true;
+				ADD_ERROR(OptionalErrors && bIsModifyingDefaults, OptionalErrors->DefaultFailures.Add(TargetStreamId, EConcertChangeStreamFrequencyErrorCode::UnknownStream))
+				ADD_ERROR(OptionalErrors, MarkAsMissingStream(TargetStreamId, SubRequest.OverridesToPut, *OptionalErrors))
+				ADD_ERROR(OptionalErrors, MarkAsMissingStream(TargetStreamId, SubRequest.OverridesToAdd, *OptionalErrors))
+				continue;
+			}
+
+			// Defaults must be valid
+			if (bIsModifyingDefaults && !SubRequest.NewDefaults.IsValid())
+			{
+				bAnyErrors = true;
+				ADD_ERROR(OptionalErrors && bIsModifyingDefaults, OptionalErrors->DefaultFailures.Add(TargetStreamId, EConcertChangeStreamFrequencyErrorCode::InvalidReplicationRate))
+			}
+
+			// Object paths referenced must have properties mapped to them in the replication map
+			bAnyErrors |= ValidateAddedFrequencies(Request, SubRequest.OverridesToPut, StreamDescription->BaseDescription, OptionalErrors);
+			bAnyErrors |= ValidateAddedFrequencies(Request, SubRequest.OverridesToAdd, StreamDescription->BaseDescription, OptionalErrors);
+		}
+		
+		return bAnyErrors;
+	}
+#undef ADD_ERROR
+	
 	void IterateInvalidEntries(
 		const FObjectReplicationMap& ReplicationMap,
 		TFunctionRef<EBreakBehavior(const FSoftObjectPath&, const FReplicatedObjectInfo&)> Callback
diff --git a/Engine/Plugins/Developer/Concert/ConcertSync/ConcertSyncCore/Source/ConcertSyncCore/Private/Replication/Processing/ObjectReplicationProcessor.cpp b/Engine/Plugins/Developer/Concert/ConcertSync/ConcertSyncCore/Source/ConcertSyncCore/Private/Replication/Processing/ObjectReplicationProcessor.cpp
--- a/Engine/Plugins/Developer/Concert/ConcertSync/ConcertSyncCore/Source/ConcertSyncCore/Private/Replication/Processing/ObjectReplicationProcessor.cpp
+++ b/Engine/Plugins/Developer/Concert/ConcertSync/ConcertSyncCore/Source/ConcertSyncCore/Private/Replication/Processing/ObjectReplicationProcessor.cpp
@@ -7,14 +7,16 @@
 
 namespace UE::ConcertSyncCore
 {
-	FObjectReplicationProcessor::FObjectReplicationProcessor(TSharedRef<IReplicationDataSource> DataSource)
+	FObjectReplicationProcessor::FObjectReplicationProcessor(
+		TSharedRef<IReplicationDataSource> DataSource
+		)
 		: DataSource(MoveTemp(DataSource))
 	{}
 
-	void FObjectReplicationProcessor::ProcessObjects(float TimeBudget)
+	void FObjectReplicationProcessor::ProcessObjects(const FProcessObjectsParams& Params)
 	{
-		// TODO: Respect time budget and prioritize objects
-		DataSource->ForEachPendingObject([this](const FObjectInStreamID& ObjectInfo)
+		// TODO UE-190714: Respect time budget and prioritize objects
+		DataSource->ForEachPendingObject([this](const FReplicatedObjectId& ObjectInfo)
 		{
 			ProcessObject({ ObjectInfo });
 		});
diff --git a/Engine/Plugins/Developer/Concert/ConcertSync/ConcertSyncCore/Source/ConcertSyncCore/Private/Replication/Processing/ObjectReplicationSender.cpp b/Engine/Plugins/Developer/Concert/ConcertSync/ConcertSyncCore/Source/ConcertSyncCore/Private/Replication/Processing/ObjectReplicationSender.cpp
--- a/Engine/Plugins/Developer/Concert/ConcertSync/ConcertSyncCore/Source/ConcertSyncCore/Private/Replication/Processing/ObjectReplicationSender.cpp
+++ b/Engine/Plugins/Developer/Concert/ConcertSync/ConcertSyncCore/Source/ConcertSyncCore/Private/Replication/Processing/ObjectReplicationSender.cpp
@@ -23,9 +23,9 @@ namespace UE::ConcertSyncCore
 		, Session(MoveTemp(Session))
 	{}
 
-	void FObjectReplicationSender::ProcessObjects(float TimeBudget)
+	void FObjectReplicationSender::ProcessObjects(const FProcessObjectsParams& Params)
 	{
-		FObjectReplicationProcessor::ProcessObjects(TimeBudget);
+		FObjectReplicationProcessor::ProcessObjects(Params);
 
 		if (!EventToSend.Streams.IsEmpty())
 		{
diff --git a/Engine/Plugins/Developer/Concert/ConcertSync/ConcertSyncCore/Source/ConcertSyncCore/Private/Replication/Processing/Proxy/FrequencyTracker.cpp b/Engine/Plugins/Developer/Concert/ConcertSync/ConcertSyncCore/Source/ConcertSyncCore/Private/Replication/Processing/Proxy/FrequencyTracker.cpp
new file mode 100644
--- /dev/null
+++ b/Engine/Plugins/Developer/Concert/ConcertSync/ConcertSyncCore/Source/ConcertSyncCore/Private/Replication/Processing/Proxy/FrequencyTracker.cpp
@@ -0,0 +1,46 @@
+// Copyright Epic Games, Inc. All Rights Reserved.
+
+#include "Replication/Processing/Proxy/FrequencyTracker.h"
+
+#include "Replication/Data/ReplicationFrequencySettings.h"
+
+namespace UE::ConcertSyncCore
+{
+	ECanProcessObject FFrequencyTracker::TrackAndCheckObject(
+		const FObjectInStreamID& ObjectPath,
+		const FConcertObjectReplicationSettings& ReplicationSettings,
+		const FPlatformSecondsTimestamp NowAsPlatformSeconds
+		)
+	{
+		// If ObjectPath is not yet tracked, this adds a FDateTime at Tick 0, which causes it to be processed immediately.
+		const FPlatformSecondsTimestamp& LastTimeProcessed = ObjectsToLastProcessed.FindOrAdd(ObjectPath);
+		return ReplicationSettings.CanProcessObject(LastTimeProcessed, NowAsPlatformSeconds);
+	}
+
+	void FFrequencyTracker::OnProcessObject(const FObjectInStreamID& Object, const FPlatformSecondsTimestamp NowAsPlatformSeconds)
+	{
+		ObjectsToLastProcessed.FindOrAdd(Object) = NowAsPlatformSeconds;
+	}
+
+	void FFrequencyTracker::CleanseOutdatedObjects(const FPlatformSecondsTimestamp CutoffTime)
+	{
+		for (auto It = ObjectsToLastProcessed.CreateIterator(); It; ++It)
+		{
+			if (It->Value <= CutoffTime)
+			{
+				It.RemoveCurrent();
+			}
+		}
+	}
+	
+	void FFrequencyTracker_CleanByTick::Tick(float DeltaTime, const FPlatformSecondsTimestamp NowAsPlatformSeconds)
+	{
+		SecondsSinceLastCleanse += DeltaTime;
+		if (SecondsSinceLastCleanse >= CleanUpIntervalSeconds)
+		{
+			SecondsSinceLastCleanse = 0.f;
+			CleanseOutdatedObjects(NowAsPlatformSeconds - CleanUpIntervalSeconds);
+		}
+	}
+}
+
diff --git a/Engine/Plugins/Developer/Concert/ConcertSync/ConcertSyncCore/Source/ConcertSyncCore/Private/Replication/Processing/ReplicationDataQueuer.cpp b/Engine/Plugins/Developer/Concert/ConcertSync/ConcertSyncCore/Source/ConcertSyncCore/Private/Replication/Processing/ReplicationDataQueuer.cpp
--- a/Engine/Plugins/Developer/Concert/ConcertSync/ConcertSyncCore/Source/ConcertSyncCore/Private/Replication/Processing/ReplicationDataQueuer.cpp
+++ b/Engine/Plugins/Developer/Concert/ConcertSync/ConcertSyncCore/Source/ConcertSyncCore/Private/Replication/Processing/ReplicationDataQueuer.cpp
@@ -14,7 +14,7 @@ namespace UE::ConcertSyncCore
 		}
 	}
 
-	void FReplicationDataQueuer::ForEachPendingObject(TFunctionRef<void(const FObjectInStreamID&)> ProcessItemFunc) const
+	void FReplicationDataQueuer::ForEachPendingObject(TFunctionRef<void(const FReplicatedObjectId&)> ProcessItemFunc) const
 	{
 		for (auto It = PendingEvents.CreateConstIterator(); It; ++It)
 		{
@@ -28,7 +28,7 @@ namespace UE::ConcertSyncCore
 	}
 
 	bool FReplicationDataQueuer::ExtractReplicationDataForObject(
-		const FObjectInStreamID& Object,
+		const FReplicatedObjectId& Object,
 		TFunctionRef<void(const FConcertSessionSerializedPayload& Payload)> ProcessCopyable,
 		TFunctionRef<void(FConcertSessionSerializedPayload&& Payload)> ProcessMoveable
 		)
diff --git a/Engine/Plugins/Developer/Concert/ConcertSync/ConcertSyncCore/Source/ConcertSyncCore/Public/Replication/ChangeStreamSharedUtils.h b/Engine/Plugins/Developer/Concert/ConcertSync/ConcertSyncCore/Source/ConcertSyncCore/Public/Replication/ChangeStreamSharedUtils.h
--- a/Engine/Plugins/Developer/Concert/ConcertSync/ConcertSyncCore/Source/ConcertSyncCore/Public/Replication/ChangeStreamSharedUtils.h
+++ b/Engine/Plugins/Developer/Concert/ConcertSync/ConcertSyncCore/Source/ConcertSyncCore/Public/Replication/ChangeStreamSharedUtils.h
@@ -4,6 +4,7 @@
 
 #include "Containers/Array.h"
 #include "HAL/Platform.h"
+#include "Messages/ChangeStream.h"
 #include "Templates/Function.h"
 
 enum class EBreakBehavior : uint8;
@@ -39,6 +40,34 @@ namespace UE::ConcertSyncCore::Replication::ChangeStreamUtils
 		const FConcertReplication_ChangeStream_Request& Request,
 		IN OUT TArray<FReplicationStreamDescription>& StreamsToModify
 		);
+	
+	/**
+	 * Modifies StreamToModify as described by the request.
+	 * 
+	 * Validate the request before calling.
+	 * This function does use ensure conditions where it makes sense but it does not check all failure cases.
+	 * 
+	 * @param Request The request to parse
+	 * @param SettingsToModify The streams to apply the request to
+	 */
+	CONCERTSYNCCORE_API void ApplyValidatedFrequencyChanges(
+		const FConcertReplication_ChangeStream_Frequency& Request,
+		FConcertStreamFrequencySettings& SettingsToModify
+		);
+	/**
+	 * Validates a FConcertReplication_ChangeStream_Request and optionally outputs the errors into OptionalErrors if set.
+	 *
+	 * @param Request The request to validate. Note that the full request is needed (if an object is added through this
+	 * @param Streams The streams array Request would be run on
+	 * @param OptionalErrors Errors are written into this struct if set
+	 *
+	 * @return Whether Request is valid to apply to Streams
+	 */
+	CONCERTSYNCCORE_API bool ValidateFrequencyChanges(
+		const FConcertReplication_ChangeStream_Request& Request,
+		const TArray<FReplicationStreamDescription>& Streams,
+		OUT FConcertReplication_ChangeStream_FrequencyResponse* OptionalErrors = nullptr
+		);
 
 	/** Iterates entries with incomplete data, which are entries for which IsValidForSendingToServer returns false. */
 	CONCERTSYNCCORE_API void IterateInvalidEntries(const FObjectReplicationMap& ReplicationMap, TFunctionRef<EBreakBehavior(const FSoftObjectPath&, const FReplicatedObjectInfo&)> Callback);
diff --git a/Engine/Plugins/Developer/Concert/ConcertSync/ConcertSyncCore/Source/ConcertSyncCore/Public/Replication/Data/ClientQueriedInfo.h b/Engine/Plugins/Developer/Concert/ConcertSync/ConcertSyncCore/Source/ConcertSyncCore/Public/Replication/Data/ClientQueriedInfo.h
--- a/Engine/Plugins/Developer/Concert/ConcertSync/ConcertSyncCore/Source/ConcertSyncCore/Public/Replication/Data/ClientQueriedInfo.h
+++ b/Engine/Plugins/Developer/Concert/ConcertSync/ConcertSyncCore/Source/ConcertSyncCore/Public/Replication/Data/ClientQueriedInfo.h
@@ -34,6 +34,7 @@ struct FReplicationClientQueriedInfo
 	 * 
 	 * If EConcertQueryClientStreamFlags::SkipStreamInfo was set, this is empty.
 	 * If EConcertQueryClientStreamFlags::SkipProperties was set, streams' FReplicatedObjectInfo::PropertySelection are empty.
+	 * If EConcertQueryClientStreamFlags::SkipFrequency was set, the FrequencySettings are empty.
 	 */
 	UPROPERTY()
 	TArray<FSharedReplicationStreamDescription> Streams;
diff --git a/Engine/Plugins/Developer/Concert/ConcertSync/ConcertSyncCore/Source/ConcertSyncCore/Public/Replication/Data/ConcertObjectReplicationSettings.h b/Engine/Plugins/Developer/Concert/ConcertSync/ConcertSyncCore/Source/ConcertSyncCore/Public/Replication/Data/ConcertObjectReplicationSettings.h
deleted file mode 100644
diff --git a/Engine/Plugins/Developer/Concert/ConcertSync/ConcertSyncCore/Source/ConcertSyncCore/Public/Replication/Data/ObjectReplicationMap.h b/Engine/Plugins/Developer/Concert/ConcertSync/ConcertSyncCore/Source/ConcertSyncCore/Public/Replication/Data/ObjectReplicationMap.h
--- a/Engine/Plugins/Developer/Concert/ConcertSync/ConcertSyncCore/Source/ConcertSyncCore/Public/Replication/Data/ObjectReplicationMap.h
+++ b/Engine/Plugins/Developer/Concert/ConcertSync/ConcertSyncCore/Source/ConcertSyncCore/Public/Replication/Data/ObjectReplicationMap.h
@@ -10,10 +10,12 @@ USTRUCT()
 struct FReplicatedObjectInfo
 {
 	GENERATED_BODY()
-	
+
+	/** The exact class of the object being replicated. */
 	UPROPERTY()
 	FSoftClassPath ClassPath;
-	
+
+	/** The properties to replicate for the object */
 	UPROPERTY()
 	FConcertPropertySelection PropertySelection;
 
@@ -24,7 +26,11 @@ struct FReplicatedObjectInfo
 	}
 
 	/** @return Whether this data is valid for sending to the server. */
-	bool IsValidForSendingToServer() const { return ClassPath.IsValid() && !PropertySelection.ReplicatedProperties.IsEmpty(); }
+	bool IsValidForSendingToServer() const
+	{
+		return ClassPath.IsValid()
+			&& !PropertySelection.ReplicatedProperties.IsEmpty();
+	}
 	
 	friend bool operator==(const FReplicatedObjectInfo& Left, const FReplicatedObjectInfo& Right)
 	{
diff --git a/Engine/Plugins/Developer/Concert/ConcertSync/ConcertSyncCore/Source/ConcertSyncCore/Public/Replication/Data/ReplicationFrequencySettings.h b/Engine/Plugins/Developer/Concert/ConcertSync/ConcertSyncCore/Source/ConcertSyncCore/Public/Replication/Data/ReplicationFrequencySettings.h
new file mode 100644
--- /dev/null
+++ b/Engine/Plugins/Developer/Concert/ConcertSync/ConcertSyncCore/Source/ConcertSyncCore/Public/Replication/Data/ReplicationFrequencySettings.h
@@ -0,0 +1,122 @@
+// Copyright Epic Games, Inc. All Rights Reserved.
+
+#pragma once
+
+#include "HAL/PlatformTime.h"
+#include "ReplicationFrequencySettings.generated.h"
+
+UENUM()
+enum class EConcertObjectReplicationMode : uint8
+{
+	/** Replicate at the rate specified at FReplicatedObjectInfo::ReplicationRate */
+	SpecifiedRate,
+	/** Replicate the object as often as possible: every tick. */
+	Realtime
+};
+
+namespace UE::ConcertSyncCore
+{
+	enum class ECanProcessObject : uint8
+	{
+		/** Enough time has elapsed: the object can be processed. */
+		CanProcess,
+		/** Not enough time has elapsed: wait with processing the object. */
+		TooEarly
+	};
+}
+
+/** Specifies replication frequency settings for a single object. */
+USTRUCT()
+struct CONCERTSYNCCORE_API FConcertObjectReplicationSettings
+{
+	GENERATED_BODY()
+	
+	/** Controls how often this object should be replicated. */
+	UPROPERTY()
+	EConcertObjectReplicationMode ReplicationMode = EConcertObjectReplicationMode::SpecifiedRate;
+	
+	/**
+	 * How often per second the object is supposed to be replicated per second.
+	 * 
+	 * The update in seconds is given by 1 / ReplicationRate.
+	 * The default rate of 30 results in an update interval every 0.033s
+	 */
+	UPROPERTY()
+	uint8 ReplicationRate = 30;
+
+	bool IsValid() const { return ReplicationMode == EConcertObjectReplicationMode::Realtime || ReplicationRate > 0; }
+	
+	/**
+	 * Whether an object with this config can be replicated given the current time and the time it was last replicated.
+	 * @param LastTimeProcessed The time the object was last processed - past result of FPlatformTime::Seconds()
+	 * @param Now The current time - result of FPlatformTime::Seconds(). Pass this in to avoid excessive time spent on timing
+	 */
+	UE::ConcertSyncCore::ECanProcessObject CanProcessObject(const double LastTimeProcessed, const double Now = FPlatformTime::Seconds()) const
+	{
+		switch (ReplicationMode)
+		{
+		case EConcertObjectReplicationMode::SpecifiedRate:
+			return LastTimeProcessed + GetUpdateIntervalInSeconds() <= Now
+				? UE::ConcertSyncCore::ECanProcessObject::CanProcess
+				: UE::ConcertSyncCore::ECanProcessObject::TooEarly;
+			
+		case EConcertObjectReplicationMode::Realtime: return UE::ConcertSyncCore::ECanProcessObject::CanProcess;
+		default:
+			checkNoEntry();
+			return UE::ConcertSyncCore::ECanProcessObject::CanProcess;
+		}
+	}
+	
+	/** Gets the time to wait between frames replicating an object with this config. */
+	double GetUpdateIntervalInSeconds() const
+	{
+		return ReplicationMode == EConcertObjectReplicationMode::Realtime
+			? 0.0
+			: 1.0 / static_cast<double>(ReplicationRate);
+	}
+
+	friend bool operator==(const FConcertObjectReplicationSettings& Left, const FConcertObjectReplicationSettings& Right)
+	{
+		return Left.ReplicationMode == Right.ReplicationMode
+			&& Left.ReplicationRate == Right.ReplicationRate;
+	}
+	friend bool operator!=(const FConcertObjectReplicationSettings& Left, const FConcertObjectReplicationSettings& Right)
+	{
+		return !(Left == Right);
+	}
+};
+
+/** Specifies replication frequency settings for an entire stream. */
+USTRUCT()
+struct CONCERTSYNCCORE_API FConcertStreamFrequencySettings
+{
+	GENERATED_BODY()
+
+	/** Settings for all objects by default. */
+	UPROPERTY()
+	FConcertObjectReplicationSettings Defaults;
+	
+	/**
+	 * Special settings to use for specific objects. Overrides DefaultSettings.
+	 * May contain the same values as in DefaultSettings.
+	 * Only contains entries for objects that are registered in the corresponding stream.
+	 */
+	UPROPERTY()
+	TMap<FSoftObjectPath, FConcertObjectReplicationSettings> ObjectOverrides;
+
+	const FConcertObjectReplicationSettings& GetSettingsFor(const FSoftObjectPath& Object) const
+	{
+		const FConcertObjectReplicationSettings* Settings = ObjectOverrides.Find(Object);
+		return Settings ? *Settings : Defaults;
+	}
+
+	friend bool operator==(const FConcertStreamFrequencySettings& Left, const FConcertStreamFrequencySettings& Right)
+	{
+		return Left.Defaults == Right.Defaults
+		&& Left.ObjectOverrides.OrderIndependentCompareEqual(Right.ObjectOverrides);
+	}
+	friend bool operator!=(const FConcertStreamFrequencySettings& Left, const FConcertStreamFrequencySettings& Right)
+	{
+		return !(Left == Right);
+	}
+};
\ No newline at end of file
diff --git a/Engine/Plugins/Developer/Concert/ConcertSync/ConcertSyncCore/Source/ConcertSyncCore/Public/Replication/Data/ReplicationStreamDescription.h b/Engine/Plugins/Developer/Concert/ConcertSync/ConcertSyncCore/Source/ConcertSyncCore/Public/Replication/Data/ReplicationStreamDescription.h
--- a/Engine/Plugins/Developer/Concert/ConcertSync/ConcertSyncCore/Source/ConcertSyncCore/Public/Replication/Data/ReplicationStreamDescription.h
+++ b/Engine/Plugins/Developer/Concert/ConcertSync/ConcertSyncCore/Source/ConcertSyncCore/Public/Replication/Data/ReplicationStreamDescription.h
@@ -3,11 +3,23 @@
 #pragma once
 
 #include "ObjectReplicationMap.h"
+#include "ReplicationFrequencySettings.h"
 #include "ReplicationStreamDescription.generated.h"
 
 template<typename OptionalType>
 struct TOptional;
 
+namespace UE::ConcertSyncCore
+{
+	enum class EReplicationStreamCloneFlags : uint8
+	{
+		None,
+		SkipProperties = 1 << 0,
+		SkipFrequency = 1 << 1
+	};
+	ENUM_CLASS_FLAGS(EReplicationStreamCloneFlags)
+}
+
 /** Shared non-UObject data, which can be serialized by Concert without wrapping (UObjects need wrapping). */
 USTRUCT()
 struct FSharedReplicationStreamDescription
@@ -21,10 +33,33 @@ struct FSharedReplicationStreamDescription
 	/** Identifies the data that this stream will send. */
 	UPROPERTY()
 	FObjectReplicationMap ReplicationMap;
+
+	/**
+	 * Determines how often updates are sent for objects.
+	 *
+	 * The purposes of registering frequency with the server communicates the sending client's intent and has these advantages:
+	 * - If the server receives replicated data more often than specified here, it will combine and throttle the events forwarded to the other clients.
+	 * - The receiving clients can query how often per frame to expect updates for an object
+	 *
+	 * On a technical level, this field is not required for replicating objects: in a hypothetical, alternate implementation, the sending client could
+	 * just keep the replication frequency local to themselves and send at the rate it wants to send.
+	 * This field is here mostly to allow inspection by other parties, which does add overhead (engineering and runtime) because we must support
+	 * FConcertReplication_ChangeStream_Frequency requests).
+	 * 
+	 * With that said, the sending implementation uses the rate that is specified here to throttle the rate at which it sends update to the server.
+	 * @see FObjectReplicationProcessor::ProcessObjects.
+	 */
+	UPROPERTY()
+	FConcertStreamFrequencySettings FrequencySettings;
+
+	/** Util for efficiently cloning stream info and optionally skipping some info. */
+	FSharedReplicationStreamDescription Clone(UE::ConcertSyncCore::EReplicationStreamCloneFlags Flags = UE::ConcertSyncCore::EReplicationStreamCloneFlags::None) const;
 	
 	friend bool operator==(const FSharedReplicationStreamDescription& Left, const FSharedReplicationStreamDescription& Right)
 	{
-		return Left.Identifier == Right.Identifier && Left.ReplicationMap == Right.ReplicationMap;
+		return Left.Identifier == Right.Identifier
+			&& Left.ReplicationMap == Right.ReplicationMap
+			&& Left.FrequencySettings == Right.FrequencySettings;
 	}
 	friend bool operator!=(const FSharedReplicationStreamDescription& Left, const FSharedReplicationStreamDescription& Right)
 	{
@@ -34,7 +69,9 @@ struct FSharedReplicationStreamDescription
 
 /**
  * Contains a description of a replication stream.
- * A replication stream is a map of objects to their replicated properties, and a set of stream attributes (for rules).
+ * A replication stream
+ * - is a map of objects to their replicated properties,
+ * - a set of stream attributes (for rules, TODO UE-190167)
  */
 USTRUCT()
 struct FReplicationStreamDescription
@@ -80,4 +117,32 @@ struct FReplicationStreamDescription_NetPacked
 
 	/** Unpacks this data. This function can fail, e.g if the UObject class does not exist locally. */
 	CONCERTSYNCCORE_API TOptional<FReplicationStreamDescription> Unpack(FString* OutErrorMessage = nullptr) const;
-};
\ No newline at end of file
+};
+
+inline FSharedReplicationStreamDescription FSharedReplicationStreamDescription::Clone(UE::ConcertSyncCore::EReplicationStreamCloneFlags Flags) const
+{
+	using namespace UE::ConcertSyncCore;
+	FSharedReplicationStreamDescription Result { Identifier };
+
+	const TMap<FSoftObjectPath, FReplicatedObjectInfo>& ReplicatedObjects = ReplicationMap.ReplicatedObjects;
+	if (EnumHasAnyFlags(Flags, EReplicationStreamCloneFlags::SkipProperties)
+		&& !ReplicationMap.IsEmpty())
+	{
+		Result.ReplicationMap.ReplicatedObjects.Reserve(ReplicatedObjects.Num());
+		for (const TPair<FSoftObjectPath, FReplicatedObjectInfo>& ObjectInfo : ReplicatedObjects)
+		{
+			Result.ReplicationMap.ReplicatedObjects.Add(ObjectInfo.Key, { ObjectInfo.Value.ClassPath });
+		}
+	}
+	else
+	{
+		Result.ReplicationMap = ReplicationMap;
+	}
+
+	if (!EnumHasAnyFlags(Flags, EReplicationStreamCloneFlags::SkipFrequency))
+	{
+		Result.FrequencySettings = FrequencySettings;
+	}
+		
+	return Result;
+}
\ No newline at end of file
diff --git a/Engine/Plugins/Developer/Concert/ConcertSync/ConcertSyncCore/Source/ConcertSyncCore/Public/Replication/Messages/ChangeStream.h b/Engine/Plugins/Developer/Concert/ConcertSync/ConcertSyncCore/Source/ConcertSyncCore/Public/Replication/Messages/ChangeStream.h
--- a/Engine/Plugins/Developer/Concert/ConcertSync/ConcertSyncCore/Source/ConcertSyncCore/Public/Replication/Messages/ChangeStream.h
+++ b/Engine/Plugins/Developer/Concert/ConcertSync/ConcertSyncCore/Source/ConcertSyncCore/Public/Replication/Messages/ChangeStream.h
@@ -20,7 +20,7 @@ struct FConcertReplication_ChangeStream_PutObject
 {
 	GENERATED_BODY()
 
-	// TODO UE-201166: Add two enums here that describe how Propertiy and ClassPath are to be interpreted.
+	// TODO UE-201166: Add enum flags here that describe how Properties and ClassPath are to be interpreted.
 	
 	/**
 	 * The property selection the object should have.
@@ -62,6 +62,57 @@ struct FConcertReplication_ChangeStream_PutObject
 	CONCERTSYNCCORE_API TOptional<FReplicatedObjectInfo> MakeObjectInfoIfValid() const;
 };
 
+UENUM()
+enum class EConcertReplicationChangeFrequencyFlags : uint8
+{
+	None = 0,
+	/** FConcertReplication_ChangeStream_Frequency::NewDefaults contains the new defaults to set. */
+	SetDefaults = 1 << 0,
+};
+ENUM_CLASS_FLAGS(EConcertReplicationChangeFrequencyFlags);
+
+/** Changes to set to make to the stream's replication frequency. */
+USTRUCT()
+struct FConcertReplication_ChangeStream_Frequency
+{
+	GENERATED_BODY()
+
+	/**
+	 * Replaces the entire overrides with this content.
+	 * If a specified object is not specified in the associated stream, this request DOES fail.
+	 * 
+	 * Usually you want to specify either OverridesToPut, or OverridesToPut & OverridesToRemove but not both.
+	 * A put is useful for situations where you don't care what objects have been overriden thus far and want to straight out replace them all.
+	 */
+	UPROPERTY()
+	TMap<FSoftObjectPath, FConcertObjectReplicationSettings> OverridesToPut;
+
+	/**
+	 * Removes the specified overrides.
+	 * If a specified object is not specified in the associated stream, this request DOES NOT fail.
+	 * 
+	 * Executed after OverridesToPut (but you should either use OverridesToPut, or OverridesToRemove & OverridesToAdd).
+	 */
+	UPROPERTY()
+	TSet<FSoftObjectPath> OverridesToRemove;
+	
+	/**
+	 * Adds the specified overrides.
+	 * If a specified object is not specified in the associated stream, this request DOES fail.
+	 * 
+	 * Executed after OverridesToAdd (but you should either use OverridesToPut, or OverridesToRemove & OverridesToAdd).
+	 */
+	UPROPERTY()
+	TMap<FSoftObjectPath, FConcertObjectReplicationSettings> OverridesToAdd;
+
+	/** Only applied if Flags contains the SetDefaults flag. */
+	UPROPERTY()
+	FConcertObjectReplicationSettings NewDefaults;
+	
+	UPROPERTY()
+	EConcertReplicationChangeFrequencyFlags Flags = EConcertReplicationChangeFrequencyFlags::None;
+};
+
 /**
  * Let's a client change its streams owned on the server.
  * 
@@ -91,14 +142,27 @@ struct FConcertReplication_ChangeStream_Request
 	 * @see FConcertChangeStream_Response::AuthorityConflicts for some examples of conflicts.
 	 *
 	 * If the key identifies a stream that does not exist, the request will fail.
-	 * 
+	 *
+	 * TODO UE-201167: Make the request no longer fail
 	 * If the key identifies a stream that is in StreamsToAdd, the request will fail to avoid allowing the construction of
 	 * ambiguous requests, e.g. both StreamsToAdd and ObjectsToPut containing object Foo but with different property selections.
 	 */
 	UPROPERTY()
 	TMap<FObjectInStreamID, FConcertReplication_ChangeStream_PutObject> ObjectsToPut;
 
-	// TODO UE-201167: Change StreamsToAdd so you just give FGuid of new streams and specify properties in ObjectsToPut > easier API usage.
+	/**
+	 * Changes the replication frequency settings for a stream.
+	 * 
+	 * If a specified stream does not exist, the request fails.
+	 * Specifying a request that changes nothing does not result in any failure.
+	 * 
+	 * TODO UE-201167: Change StreamsToAdd such that no ambiguous requests are even possible and update this doc's wording.
+	 * If the key may identify a stream that is in StreamsToAdd, in which case the frequency is what is specified in FrequencyChanges.
+	 */
+	UPROPERTY()
+	TMap<FGuid, FConcertReplication_ChangeStream_Frequency> FrequencyChanges;
+	
+	// TODO UE-201167: Change StreamsToAdd so you just give FGuid of new streams and specify properties in ObjectsToPut and FrequencyChanges > easier API usage.
 	
 	/**
 	 * New streams to add to the server.
@@ -150,6 +214,41 @@ enum class EConcertPutObjectErrorCode : uint8
 	MissingData
 };
 
+UENUM()
+enum class EConcertChangeObjectFrequencyErrorCode : uint8
+{
+	/** The object for which the frequency was being changed was not registered. */
+	NotRegistered,
+	/** The replication rate parameter was rejected (it cannot be 0). */
+	InvalidReplicationRate
+};
+
+UENUM()
+enum class EConcertChangeStreamFrequencyErrorCode : uint8
+{
+	/** The stream was not registered */
+	UnknownStream,
+	/** The replication rate parameter was rejected (it cannot be 0). */
+	InvalidReplicationRate
+};
+
+USTRUCT()
+struct FConcertReplication_ChangeStream_FrequencyResponse
+{
+	GENERATED_BODY()
+
+	/** Streams that could not have their frequencies overriden. */
+	UPROPERTY()
+	TMap<FObjectInStreamID, EConcertChangeObjectFrequencyErrorCode> OverrideFailures;
+
+	/** Streams that could not have their frequency defaults overriden. */
+	UPROPERTY()
+	TMap<FGuid, EConcertChangeStreamFrequencyErrorCode> DefaultFailures;
+	
+	bool IsSuccess() const { return OverrideFailures.IsEmpty() && DefaultFailures.IsEmpty(); }
+	bool IsFailure() const { return !IsSuccess(); }
+};
+
 /**
  * Contains information about why a request failed. This info could be parsed and displayed to the end user as error.
  * If there is even just one error, the entire request fails and no changes are made server-side.
@@ -189,7 +288,18 @@ struct FConcertReplication_ChangeStream_Response
 	UPROPERTY()
 	TSet<FGuid> FailedStreamCreation;
 
-	bool IsSuccess() const { return ErrorCode == EReplicationResponseErrorCode::Handled && AuthorityConflicts.IsEmpty() && ObjectsToPutSemanticErrors.IsEmpty() && FailedStreamCreation.IsEmpty(); }
+	/** Errors updating the frequency */
+	UPROPERTY()
+	FConcertReplication_ChangeStream_FrequencyResponse FrequencyErrors;
+	
+	bool IsSuccess() const
+	{
+		return ErrorCode == EReplicationResponseErrorCode::Handled
+			&& AuthorityConflicts.IsEmpty()
+			&& ObjectsToPutSemanticErrors.IsEmpty()
+			&& FailedStreamCreation.IsEmpty()
+			&& FrequencyErrors.IsSuccess();
+	}
 	bool IsFailure() const { return !IsSuccess(); }
 
 	bool WasObjectPutSuccessful(const FObjectInStreamID& Object) const { return !AuthorityConflicts.Contains(Object) && !ObjectsToPutSemanticErrors.Contains(Object); }
diff --git a/Engine/Plugins/Developer/Concert/ConcertSync/ConcertSyncCore/Source/ConcertSyncCore/Public/Replication/Messages/ClientQuery.h b/Engine/Plugins/Developer/Concert/ConcertSync/ConcertSyncCore/Source/ConcertSyncCore/Public/Replication/Messages/ClientQuery.h
--- a/Engine/Plugins/Developer/Concert/ConcertSync/ConcertSyncCore/Source/ConcertSyncCore/Public/Replication/Messages/ClientQuery.h
+++ b/Engine/Plugins/Developer/Concert/ConcertSync/ConcertSyncCore/Source/ConcertSyncCore/Public/Replication/Messages/ClientQuery.h
@@ -12,9 +12,7 @@ enum class EConcertQueryClientStreamFlags : uint8
 {
 	None = 0,
 
-	/**
-	 * If set, FReplicationClientQueriedInfo::Streams will be empty.
-	 */
+	/** If set, FReplicationClientQueriedInfo::Streams will be empty. */
 	SkipStreamInfo = 1 << 0,
 
 	/**
@@ -24,10 +22,14 @@ enum class EConcertQueryClientStreamFlags : uint8
 	 */
 	SkipProperties = 1 << 1,
 	
+	/** If set, FReplicationClientQueriedInfo::Authority will be empty. */
+	SkipAuthority = 1 << 2,
+
 	/**
-	 * If set, FReplicationClientQueriedInfo::Authority will be empty.
+	 * Optimization. Ignored if SkipStreamInfo is set.
+	 * If set, FReplicationClientQueriedInfo::Streams.FrequencySettings will not be set.
 	 */
-	SkipAuthority = 1 << 2
+	SkipFrequency = 1 << 3,
 };
 ENUM_CLASS_FLAGS(EConcertQueryClientStreamFlags)
 
diff --git a/Engine/Plugins/Developer/Concert/ConcertSync/ConcertSyncCore/Source/ConcertSyncCore/Public/Replication/Processing/IReplicationDataSource.h b/Engine/Plugins/Developer/Concert/ConcertSync/ConcertSyncCore/Source/ConcertSyncCore/Public/Replication/Processing/IReplicationDataSource.h
--- a/Engine/Plugins/Developer/Concert/ConcertSync/ConcertSyncCore/Source/ConcertSyncCore/Public/Replication/Processing/IReplicationDataSource.h
+++ b/Engine/Plugins/Developer/Concert/ConcertSync/ConcertSyncCore/Source/ConcertSyncCore/Public/Replication/Processing/IReplicationDataSource.h
@@ -7,7 +7,7 @@
 #include "Misc/Optional.h"
 #include "Templates/Function.h"
 
-struct FObjectInStreamID;
+struct FReplicatedObjectId;
 
 namespace UE::ConcertSyncCore
 {
@@ -23,7 +23,7 @@ namespace UE::ConcertSyncCore
 	public:
 		
 		/** Iterates the objects that must be processed for replication (the result of ExtractReplicationDataForObject MAY return something new). */
-		virtual void ForEachPendingObject(TFunctionRef<void(const FObjectInStreamID&)> ProcessItemFunc) const = 0;
+		virtual void ForEachPendingObject(TFunctionRef<void(const FReplicatedObjectId&)> ProcessItemFunc) const = 0;
 		/** @return The number of objects ForEachPendingObject would iterate. Can be used, e.g. for Reserve()-ing a container. */
 		virtual int32 NumObjects() const = 0;
 
@@ -40,12 +40,12 @@ namespace UE::ConcertSyncCore
 		 * @return Whether successful. False indicates the call was invalid to make (ForEachPendingObject did not return Object). True indicates success, even if ProcessCopyable was not called.		
 		 */
 		virtual bool ExtractReplicationDataForObject(
-			const FObjectInStreamID& Object,
+			const FReplicatedObjectId& Object,
 			TFunctionRef<void(const FConcertSessionSerializedPayload& Payload)> ProcessCopyable,
 			TFunctionRef<void(FConcertSessionSerializedPayload&& Payload)> ProcessMoveable
 			) = 0;
 		/** Util version for callers that only want to read and do not want to store the payload. */
-		bool ExtractReplicationDataForObject(const FObjectInStreamID& Object, TFunctionRef<void(const FConcertSessionSerializedPayload& Payload)> ProcessCopyable)
+		bool ExtractReplicationDataForObject(const FReplicatedObjectId& Object, TFunctionRef<void(const FConcertSessionSerializedPayload& Payload)> ProcessCopyable)
 		{
 			return ExtractReplicationDataForObject(Object, ProcessCopyable, [&ProcessCopyable](FConcertSessionSerializedPayload&& Payload){ ProcessCopyable(Payload); });
 		}
diff --git a/Engine/Plugins/Developer/Concert/ConcertSync/ConcertSyncCore/Source/ConcertSyncCore/Public/Replication/Processing/ObjectReplicationProcessor.h b/Engine/Plugins/Developer/Concert/ConcertSync/ConcertSyncCore/Source/ConcertSyncCore/Public/Replication/Processing/ObjectReplicationProcessor.h
--- a/Engine/Plugins/Developer/Concert/ConcertSync/ConcertSyncCore/Source/ConcertSyncCore/Public/Replication/Processing/ObjectReplicationProcessor.h
+++ b/Engine/Plugins/Developer/Concert/ConcertSync/ConcertSyncCore/Source/ConcertSyncCore/Public/Replication/Processing/ObjectReplicationProcessor.h
@@ -4,40 +4,71 @@
 
 #include "IReplicationDataSource.h"
 #include "Replication/Data/ObjectIds.h"
+
 #include "Templates/SharedPointer.h"
 #include "UObject/SoftObjectPath.h"
 
+struct FConcertObjectReplicationSettings;
+struct FObjectInStreamID;
+
 namespace UE::ConcertSyncCore
 {
 	class IReplicationDataSource;
+		
+	/**
+	 * Params for FObjectReplicationProcessor::ProcessObjects.
+	 * Exists as struct so we can add params (e.g. UE-190714) without affecting the signature and can avoid deprecation.
+	 */
+	struct FProcessObjectsParams
+	{
+		/** The current frame's delta time */
+		float DeltaTime;
+		// TODO UE-190714: We should add a time budget.
+	};
 	
+	
+	class CONCERTSYNCCORE_API IObjectReplicationProcessor
+	{
+	public:
+		
+		/**
+		 * Processes all changed objects under the given time budget.
+		 * TODO UE-190714: A time budget should be added.
+		 */
+		virtual void ProcessObjects(const FProcessObjectsParams& Params) = 0;
+
+		virtual ~IObjectReplicationProcessor() = default;
+	};
+
+
 	/**
 	 * Responsible for prioritizing a list of objects and processing them.
 	 * 
 	 * Processing is defined by the subclass. Some example implementations:
 	 *  - serialize the object and send (client),
-	 *  - find loaded object and apply replication data (client)
-	 *  - send object data to interested clients (server),
+	 *  - find loaded object and apply replication data (client),
+	 *  - send object data to interested clients (server)
+	 *  
+	 * FObjectReplicationProcessor is backed by an IReplicationDataSource, which is a source of replicated object data.
 	 */
 	class CONCERTSYNCCORE_API FObjectReplicationProcessor
 	{
 	public:
-
-		explicit FObjectReplicationProcessor(TSharedRef<IReplicationDataSource> DataSource);
-		virtual ~FObjectReplicationProcessor() = default;
-
+		
 		/**
-		 * Processes all changed objects under the given time budget.
-		 * The time budget may be exceeded but we'll try not to and to stay as close to the budget as possible.
+		 * @param DataSource Source of the data that is to be sent
 		 */
-		virtual void ProcessObjects(float TimeBudget);
+		FObjectReplicationProcessor(TSharedRef<IReplicationDataSource> DataSource);
+		virtual ~FObjectReplicationProcessor() = default;
+		
+		virtual void ProcessObjects(const FProcessObjectsParams& Params);
 		
 	protected:
 
 		struct FObjectProcessArgs
 		{
 			/** Info about the object to process */
-			FObjectInStreamID ObjectInfo;
+			FReplicatedObjectId ObjectInfo;
 		};
 
 		FORCEINLINE IReplicationDataSource& GetDataSource() const { return DataSource.Get(); }
@@ -48,6 +79,6 @@ namespace UE::ConcertSyncCore
 	private:
 		
 		/** Abstracts where replication data comes from: could be generated (clients) or received (server or client) */
-		TSharedRef<IReplicationDataSource> DataSource;
+		const TSharedRef<IReplicationDataSource> DataSource;
 	};
 }
diff --git a/Engine/Plugins/Developer/Concert/ConcertSync/ConcertSyncCore/Source/ConcertSyncCore/Public/Replication/Processing/ObjectReplicationSender.h b/Engine/Plugins/Developer/Concert/ConcertSync/ConcertSyncCore/Source/ConcertSyncCore/Public/Replication/Processing/ObjectReplicationSender.h
--- a/Engine/Plugins/Developer/Concert/ConcertSync/ConcertSyncCore/Source/ConcertSyncCore/Public/Replication/Processing/ObjectReplicationSender.h
+++ b/Engine/Plugins/Developer/Concert/ConcertSync/ConcertSyncCore/Source/ConcertSyncCore/Public/Replication/Processing/ObjectReplicationSender.h
@@ -15,11 +15,20 @@ namespace UE::ConcertSyncCore
 	class CONCERTSYNCCORE_API FObjectReplicationSender : public FObjectReplicationProcessor
 	{
 	public:
-
-		FObjectReplicationSender(const FGuid& TargetEndpointId, TSharedRef<IConcertSession> Session, TSharedRef<IReplicationDataSource> DataSource);
+		
+		/**
+		 * @param TargetEndpointId The endpoints to send to
+		 * @param Session The session to use for sending
+		 * @param DataSource Source of the data that is to be sent
+		 */
+		FObjectReplicationSender(
+			const FGuid& TargetEndpointId,
+			TSharedRef<IConcertSession> Session,
+			TSharedRef<IReplicationDataSource> DataSource
+			);
 		
 		//~ Begin FObjectReplicationProcessor Interface
-		virtual void ProcessObjects(float TimeBudget) override;
+		virtual void ProcessObjects(const FProcessObjectsParams& Params) override;
 		//~ End FObjectReplicationProcessor Interface
 
 	protected:
@@ -32,9 +41,9 @@ namespace UE::ConcertSyncCore
 
 		/** The endpoint data will be sent to */
 		const FGuid TargetEndpointId;
-
+		
 		/** The session through which replication messages are sent. */
-		TSharedRef<IConcertSession> Session;
+		const TSharedRef<IConcertSession> Session;
 
 		/** This event is filled in ProcessObjects and finally sent to TargetEndpointId. */
 		FConcertReplication_BatchReplicationEvent EventToSend;
diff --git a/Engine/Plugins/Developer/Concert/ConcertSync/ConcertSyncCore/Source/ConcertSyncCore/Public/Replication/Processing/Proxy/FrequencyTracker.h b/Engine/Plugins/Developer/Concert/ConcertSync/ConcertSyncCore/Source/ConcertSyncCore/Public/Replication/Processing/Proxy/FrequencyTracker.h
new file mode 100644
--- /dev/null
+++ b/Engine/Plugins/Developer/Concert/ConcertSync/ConcertSyncCore/Source/ConcertSyncCore/Public/Replication/Processing/Proxy/FrequencyTracker.h
@@ -0,0 +1,92 @@
+// Copyright Epic Games, Inc. All Rights Reserved.
+
+#pragma once
+
+#include "Replication/Data/ObjectIds.h"
+#include "Replication/Data/ReplicationFrequencySettings.h"
+
+#include "Containers/Map.h"
+#include "HAL/PlatformTime.h"
+
+struct FDateTime;
+
+namespace UE::ConcertSyncCore
+{
+	/**
+	 * Keeps track of when objects are processed.
+	 * Allows querying of whether in a given frame an object is allowed to be processed (or whether it is too early).
+	 *
+	 * This keeps a TMap<FSoftObjectPath, FGuid> of objects that were processed in the past.
+	 * You must periodically call CleanseOutdatedObjects() to avoid leaking objects that are no longer sent.
+	 * @see FFrequencyTracker_CleanByTick
+	 */
+	class CONCERTSYNCCORE_API FFrequencyTracker
+	{
+	public:
+
+		// Holds the return value of FPlatformTime::Seconds()
+		using FPlatformSecondsTimestamp = double;
+
+		/**
+		 * Ensures ObjectPath is tracked in the cache and returns whether it can be processed now given its frequency settings.
+		 *
+		 * Note that multiple streams may be replicating the object and each stream has a different frequency settings for an object, e.g.
+		 * object StaticMeshActor in stream A may be replicating RelativeLocation at 60 FPS and in stream B may be replicating RelativeRotation at 30 FPS (contrived example).
+		 *
+		 * The object can send immediately when it is first added, i.e. TrackAndCheckObject returns ECanProcessObject::CanProcess if the object was not yet registered.
+		 * 
+		 * @param ObjectPath The object that should be processed
+		 * @param ReplicationSettings The frequency settings for the object.
+		 * @param NowAsPlatformSeconds The current time - result of FPlatformTime::Seconds(). Pass this in to avoid excessive time spent on timing.
+		 * 
+		 * @return Whether the object can be processed now
+		 */
+		ECanProcessObject TrackAndCheckObject(
+			const FObjectInStreamID& ObjectPath,
+			const FConcertObjectReplicationSettings& ReplicationSettings,
+			const FPlatformSecondsTimestamp NowAsPlatformSeconds = FPlatformTime::Seconds()
+			);
+
+		/**
+		 * Called when an object is processed.
+		 * This is allowed to be called even if TrackAndCheckObject reported it being too early to process the object.
+		 *
+		 * @param Object The object that was processed
+		 * @param NowAsPlatformSeconds The current time - result of FPlatformTime::Seconds(). Pass this in to avoid excessive time spent on timing.
+		 */
+		void OnProcessObject(const FObjectInStreamID& Object, const FPlatformSecondsTimestamp NowAsPlatformSeconds = FPlatformTime::Seconds());
+		
+		/**
+		 * Removes all tracked objects were last processed before CutoffTime.
+		 * @param CutoffTime A past value of FPlatformTime::Seconds()
+		 */
+		void CleanseOutdatedObjects(const FPlatformSecondsTimestamp CutoffTime);
+
+	protected:
+
+		// For unit tests
+		const TMap<FObjectInStreamID, FPlatformSecondsTimestamp>& GetObjectsToLastProcessed() const { return ObjectsToLastProcessed; }
+		
+	private:
+
+		/** It is important to call CleanseOutdatedObjects to avoid this mapping leaking objects that used to be replicated. */
+		TMap<FObjectInStreamID, FPlatformSecondsTimestamp> ObjectsToLastProcessed;
+	};
+
+	/** Subclass of FFrequencyTracker that you call Tick on. Every x seconds, this will call CleanseOutdatedObjects itself. */
+	class CONCERTSYNCCORE_API FFrequencyTracker_CleanByTick : public FFrequencyTracker
+	{
+	public:
+
+		FFrequencyTracker_CleanByTick(const double InCleanUpInterval = 5.0)
+			: CleanUpIntervalSeconds(InCleanUpInterval)
+		{}
+
+		void Tick(float DeltaTime, const FPlatformSecondsTimestamp NowAsPlatformSeconds = FPlatformTime::Seconds());
+
+	private:
+
+		const double CleanUpIntervalSeconds;
+		double SecondsSinceLastCleanse = 0.f;
+	};
+}
diff --git a/Engine/Plugins/Developer/Concert/ConcertSync/ConcertSyncCore/Source/ConcertSyncCore/Public/Replication/Processing/Proxy/ObjectProcessorProxy.h b/Engine/Plugins/Developer/Concert/ConcertSync/ConcertSyncCore/Source/ConcertSyncCore/Public/Replication/Processing/Proxy/ObjectProcessorProxy.h
new file mode 100644
--- /dev/null
+++ b/Engine/Plugins/Developer/Concert/ConcertSync/ConcertSyncCore/Source/ConcertSyncCore/Public/Replication/Processing/Proxy/ObjectProcessorProxy.h
@@ -0,0 +1,39 @@
+// Copyright Epic Games, Inc. All Rights Reserved.
+
+#pragma once
+
+#include "Replication/Processing/ObjectReplicationProcessor.h"
+
+#include <type_traits>
+
+namespace UE::ConcertSyncCore
+{
+	template <typename T>
+	concept TIsObjectProcessorConcept = std::is_base_of_v<FObjectReplicationProcessor, T>;
+	
+    /**
+     * A proxy processor adds additional operations on top of FObjectReplicationProcessor at compile time.
+     * The advantage of this approach is that multiple behaviours can flexibly be chained on top of a real processor implementation with minimal overhead.
+     * 
+     * This class is supposed to be subclassed and call functions on the real processor implementation.
+     * The point of this class is to be a marker interface (i.e. when somebody reads TObjectProcessorProxy they know what the intention is).
+     */
+    template<TIsObjectProcessorConcept TRealProcessorImpl>
+    class TObjectProcessorProxy : public TRealProcessorImpl
+    {
+    public:
+    	
+	    /**
+		 * @param Arg Perfectly forwarded args for constructing the real processor implementation
+	     */
+	    template<typename... TArg>
+    	TObjectProcessorProxy(TArg&&... Arg)
+    		: TRealProcessorImpl(Forward<TArg>(Arg)...)
+    	{}
+
+    protected:
+
+    	const TRealProcessorImpl& GetInnerProcessor() const { return *this; }
+    	TRealProcessorImpl& GetInnerProcessor() { return *this; }
+    };
+}
diff --git a/Engine/Plugins/Developer/Concert/ConcertSync/ConcertSyncCore/Source/ConcertSyncCore/Public/Replication/Processing/Proxy/ObjectProcessorProxy_Frequency.h b/Engine/Plugins/Developer/Concert/ConcertSync/ConcertSyncCore/Source/ConcertSyncCore/Public/Replication/Processing/Proxy/ObjectProcessorProxy_Frequency.h
new file mode 100644
--- /dev/null
+++ b/Engine/Plugins/Developer/Concert/ConcertSync/ConcertSyncCore/Source/ConcertSyncCore/Public/Replication/Processing/Proxy/ObjectProcessorProxy_Frequency.h
@@ -0,0 +1,98 @@
+// Copyright Epic Games, Inc. All Rights Reserved.
+
+#pragma once
+
+#include "FrequencyTracker.h"
+#include "ObjectProcessorProxy.h"
+#include "Replication/Data/ReplicationFrequencySettings.h"
+
+#include "Delegates/Delegate.h"
+#include "HAL/PlatformTime.h"
+
+namespace UE::ConcertSyncCore
+{
+	DECLARE_DELEGATE_RetVal_OneParam(FConcertObjectReplicationSettings, FGetObjectFrequencySettings, const FReplicatedObjectId& Object);
+	
+	/**
+	 * Maps every FSoftObjectPath to FConcertObjectReplicationSettings.
+	 * Extends the inner IObjectReplicationProcessor by throttling the processing of objects according to the associated frequency settings.  
+	 */
+	template<TIsObjectProcessorConcept TRealProcessorImpl>
+	class TObjectProcessorProxy_Frequency : public TObjectProcessorProxy<TRealProcessorImpl>
+    {
+	public:
+		
+		/**
+		 * @param GetObjectFrequencySettingsDelegate Gets the frequency settings for an object in a stream
+		 * @param Arg Perfectly forwarded args for constructing the real processor implementation
+		 */
+		template<typename... TArg>
+		TObjectProcessorProxy_Frequency(FGetObjectFrequencySettings GetObjectFrequencySettingsDelegate, TArg&&... Arg)
+			: TObjectProcessorProxy<TRealProcessorImpl>(Forward<TArg>(Arg)...)
+			, GetObjectFrequencySettingsDelegate(MoveTemp(GetObjectFrequencySettingsDelegate))
+		{}
+
+		/**
+		 * @param GetObjectFrequencySettingsDelegate Gets the frequency settings for an object in a stream
+		 * @param ObjectCacheRetentionSeconds If an object goes this time without being processed, it is removed from the cache.
+		 * @param Arg Perfectly forwarded args for constructing the real processor implementation
+		 */
+		template<typename... TArg>
+		TObjectProcessorProxy_Frequency(FGetObjectFrequencySettings GetObjectFrequencySettingsDelegate, const float ObjectCacheRetentionSeconds, TArg&&... Arg)
+			: TObjectProcessorProxy<TRealProcessorImpl>(Forward<TArg>(Arg)...)
+			, GetObjectFrequencySettingsDelegate(MoveTemp(GetObjectFrequencySettingsDelegate))
+			, FrequencyTracker(ObjectCacheRetentionSeconds)
+		{}
+
+		//~ Begin FObjectReplicationProcessor Interface
+		virtual void ProcessObjects(const FProcessObjectsParams& Params) override;
+		//~ End FObjectReplicationProcessor Interface
+
+	protected:
+		
+		//~ Begin FObjectReplicationProcessor Interface
+		virtual void ProcessObject(const FObjectReplicationProcessor::FObjectProcessArgs& Args) override;
+		//~ End FObjectReplicationProcessor Interface
+
+	private:
+
+		/** Gets an objet's frequency settings */
+		const FGetObjectFrequencySettings GetObjectFrequencySettingsDelegate;
+		
+		/**
+		 * Used for throttling object replication.
+		 * Tracks when an object is processed and can be queried whether it is time to process the object again.
+		 */
+		FFrequencyTracker_CleanByTick FrequencyTracker;
+
+		/** Updated with every invocation of ProcessObjects and used by ProcessObject. The result of FPlatformTime::Seconds(). */
+		double CurrentProcessingTime{};
+    };
+
+	template <TIsObjectProcessorConcept TRealProcessorImpl>
+	void TObjectProcessorProxy_Frequency<TRealProcessorImpl>::ProcessObjects(const FProcessObjectsParams& Params)
+	{
+		CurrentProcessingTime = FPlatformTime::Seconds();
+		// Normally calls ProcessObject on all objects that want to be replicated
+		TRealProcessorImpl::ProcessObjects(Params);
+		
+		// This removes objects from the cache that have not been replicated in a while
+		FrequencyTracker.Tick(Params.DeltaTime);
+	}
+
+	template <TIsObjectProcessorConcept TRealProcessorImpl>
+	void TObjectProcessorProxy_Frequency<TRealProcessorImpl>::ProcessObject(const FObjectReplicationProcessor::FObjectProcessArgs& Args)
+	{
+		const FReplicatedObjectId& Object = Args.ObjectInfo;
+		const FConcertObjectReplicationSettings FrequencySettings = GetObjectFrequencySettingsDelegate.Execute(Object);
+
+		// Throttle calls to the real ProcessObject implementation based on the object's frequency settings
+		if (FrequencyTracker.TrackAndCheckObject(Object, FrequencySettings, CurrentProcessingTime) == ECanProcessObject::CanProcess)
+		{
+			FrequencyTracker.OnProcessObject(Object, CurrentProcessingTime);
+			
+			// Do the actual processing of the object
+			TRealProcessorImpl::ProcessObject(Args);
+		}
+	}
+}
diff --git a/Engine/Plugins/Developer/Concert/ConcertSync/ConcertSyncCore/Source/ConcertSyncCore/Public/Replication/Processing/ReplicationDataQueuer.h b/Engine/Plugins/Developer/Concert/ConcertSync/ConcertSyncCore/Source/ConcertSyncCore/Public/Replication/Processing/ReplicationDataQueuer.h
--- a/Engine/Plugins/Developer/Concert/ConcertSync/ConcertSyncCore/Source/ConcertSyncCore/Public/Replication/Processing/ReplicationDataQueuer.h
+++ b/Engine/Plugins/Developer/Concert/ConcertSync/ConcertSyncCore/Source/ConcertSyncCore/Public/Replication/Processing/ReplicationDataQueuer.h
@@ -27,9 +27,9 @@ namespace UE::ConcertSyncCore
 		virtual ~FReplicationDataQueuer() override;
 
 		//~ Begin IReplicationDataSource Interface
-		virtual void ForEachPendingObject(TFunctionRef<void(const FObjectInStreamID&)> ProcessItemFunc) const override;
+		virtual void ForEachPendingObject(TFunctionRef<void(const FReplicatedObjectId&)> ProcessItemFunc) const override;
 		virtual int32 NumObjects() const override;
-		virtual bool ExtractReplicationDataForObject(const FObjectInStreamID& Object, TFunctionRef<void(const FConcertSessionSerializedPayload& Payload)> ProcessCopyable, TFunctionRef<void(FConcertSessionSerializedPayload&& Payload)> ProcessMoveable) override;
+		virtual bool ExtractReplicationDataForObject(const FReplicatedObjectId& Object, TFunctionRef<void(const FConcertSessionSerializedPayload& Payload)> ProcessCopyable, TFunctionRef<void(FConcertSessionSerializedPayload&& Payload)> ProcessMoveable) override;
 		//~ End IReplicationDataSource Interface
 		
 		//~ Begin IReplicationCacheUser Interface
@@ -44,7 +44,7 @@ namespace UE::ConcertSyncCore
 	private:
 
 		/** Stores events as they are received. */
-		TMap<FObjectInStreamID, TSharedPtr<const FConcertReplication_ObjectReplicationEvent>> PendingEvents;
+		TMap<FReplicatedObjectId, TSharedPtr<const FConcertReplication_ObjectReplicationEvent>> PendingEvents;
 
 		/** Provides us with replication events and shares them effectively. */
 		TSharedPtr<FObjectReplicationCache> ReplicationCache;
diff --git a/Engine/Plugins/Developer/Concert/ConcertSync/ConcertSyncServer/Source/ConcertSyncServer/Private/Replication/ConcertReplicationClient.cpp b/Engine/Plugins/Developer/Concert/ConcertSync/ConcertSyncServer/Source/ConcertSyncServer/Private/Replication/ConcertReplicationClient.cpp
--- a/Engine/Plugins/Developer/Concert/ConcertSync/ConcertSyncServer/Source/ConcertSyncServer/Private/Replication/ConcertReplicationClient.cpp
+++ b/Engine/Plugins/Developer/Concert/ConcertSync/ConcertSyncServer/Source/ConcertSyncServer/Private/Replication/ConcertReplicationClient.cpp
@@ -12,17 +12,21 @@ namespace UE::ConcertSyncServer::Replication
 		TArray<FReplicationStreamDescription> StreamDescriptions,
 		const FGuid& ClientEndpointId,
 		TSharedRef<IConcertSession> Session,
-		TSharedRef<ConcertSyncCore::FObjectReplicationCache> ReplicationCache
+		TSharedRef<ConcertSyncCore::FObjectReplicationCache> ReplicationCache,
+		ConcertSyncCore::FGetObjectFrequencySettings GetObjectFrequencySettingsDelegate
 	)
 		: StreamDescriptions(MoveTemp(StreamDescriptions))
 		, ClientEndpointId(ClientEndpointId)
 		, EventQueue(FServerReplicationDataQueuer::Make(ClientEndpointId, MoveTemp(ReplicationCache)))
-		, DataRelay(ClientEndpointId, MoveTemp(Session), EventQueue)
+		, DataRelay(
+			MoveTemp(GetObjectFrequencySettingsDelegate),
+			ClientEndpointId, MoveTemp(Session), EventQueue
+			)
 	{}
 
-	void FConcertReplicationClient::ProcessClient(float TimeBudget)
+	void FConcertReplicationClient::ProcessClient(const ConcertSyncCore::FProcessObjectsParams& Params)
 	{
-		DataRelay.ProcessObjects(TimeBudget);
+		DataRelay.ProcessObjects(Params);
 	}
 
 	void FConcertReplicationClient::ApplyValidatedRequest(const FConcertReplication_ChangeStream_Request& Request)
diff --git a/Engine/Plugins/Developer/Concert/ConcertSync/ConcertSyncServer/Source/ConcertSyncServer/Private/Replication/ConcertReplicationClient.h b/Engine/Plugins/Developer/Concert/ConcertSync/ConcertSyncServer/Source/ConcertSyncServer/Private/Replication/ConcertReplicationClient.h
--- a/Engine/Plugins/Developer/Concert/ConcertSync/ConcertSyncServer/Source/ConcertSyncServer/Private/Replication/ConcertReplicationClient.h
+++ b/Engine/Plugins/Developer/Concert/ConcertSync/ConcertSyncServer/Source/ConcertSyncServer/Private/Replication/ConcertReplicationClient.h
@@ -4,6 +4,7 @@
 
 #include "Containers/Array.h"
 #include "Replication/Processing/ObjectReplicationSender.h"
+#include "Replication/Processing/Proxy/ObjectProcessorProxy_Frequency.h"
 #include "Templates/SharedPointer.h"
 
 struct FReplicationStreamDescription;
@@ -13,6 +14,7 @@ struct FConcertReplication_ChangeStream_Request;
 namespace UE::ConcertSyncCore
 {
 	class FObjectReplicationCache;
+	struct FProcessObjectsParams;
 }
 
 namespace UE::ConcertSyncServer::Replication
@@ -28,14 +30,15 @@ namespace UE::ConcertSyncServer::Replication
 			TArray<FReplicationStreamDescription> StreamDescriptions,
 			const FGuid& ClientEndpointId,
 			TSharedRef<IConcertSession> Session,
-			TSharedRef<ConcertSyncCore::FObjectReplicationCache> ReplicationCache
+			TSharedRef<ConcertSyncCore::FObjectReplicationCache> ReplicationCache,
+			ConcertSyncCore::FGetObjectFrequencySettings GetObjectFrequencySettingsDelegate
 			);
 
 		/**
 		 * Process any latent tasks, such as processing pending events that need to be sent to the remote instance.
 		 * Process given a time budget. The time budget may be exceeded but we'll try not to and to stay as close to the budget as possible.
 		 */
-		void ProcessClient(float TimeBudget);
+		void ProcessClient(const ConcertSyncCore::FProcessObjectsParams& Params);
 		
 		/** Updates the StreamDescriptions array with the changes from Request. The request already passed validation and is valid to apply. */
 		void ApplyValidatedRequest(const FConcertReplication_ChangeStream_Request& Request);
@@ -54,7 +57,9 @@ namespace UE::ConcertSyncServer::Replication
 		/** Queues up replication data and passes it to DataRelay. */
 		TSharedRef<FServerReplicationDataQueuer> EventQueue;
 
+		/** Sends to remote endpoint and makes sure the objects are replicated at the specified frequency settings. */
+		using FDataRelayThrottledByFrequency = ConcertSyncCore::TObjectProcessorProxy_Frequency<ConcertSyncCore::FObjectReplicationSender>;
 		/** Sends data to the remote endpoint */
-		ConcertSyncCore::FObjectReplicationSender DataRelay;
+		FDataRelayThrottledByFrequency DataRelay;
 	};
 }
diff --git a/Engine/Plugins/Developer/Concert/ConcertSync/ConcertSyncServer/Source/ConcertSyncServer/Private/Replication/ConcertServerReplicationManager.cpp b/Engine/Plugins/Developer/Concert/ConcertSync/ConcertSyncServer/Source/ConcertSyncServer/Private/Replication/ConcertServerReplicationManager.cpp
--- a/Engine/Plugins/Developer/Concert/ConcertSync/ConcertSyncServer/Source/ConcertSyncServer/Private/Replication/ConcertServerReplicationManager.cpp
+++ b/Engine/Plugins/Developer/Concert/ConcertSync/ConcertSyncServer/Source/ConcertSyncServer/Private/Replication/ConcertServerReplicationManager.cpp
@@ -43,7 +43,7 @@ namespace UE::ConcertSyncServer::Replication
 
 	void FConcertServerReplicationManager::ForEachStream(const FGuid& ClientEndpointId, TFunctionRef<EBreakBehavior(const FReplicationStreamDescription& Stream)> Callback) const
 	{
-		const TSharedRef<FConcertReplicationClient>* Client = Clients.Find(ClientEndpointId);
+		const TUniquePtr<FConcertReplicationClient>* Client = Clients.Find(ClientEndpointId);
 		if (!ensure(Client))
 		{
 			return;
@@ -60,7 +60,7 @@ namespace UE::ConcertSyncServer::Replication
 
 	void FConcertServerReplicationManager::ForEachSendingClient(TFunctionRef<EBreakBehavior(const FGuid& ClientEndpointId)> Callback) const
 	{
-		for (const TPair<FGuid, TSharedRef<FConcertReplicationClient>>& ClientPair : Clients)
+		for (const TPair<FGuid, TUniquePtr<FConcertReplicationClient>>& ClientPair : Clients)
 		{
 			if (!ClientPair.Value->GetStreamDescriptions().IsEmpty()
 				&& Callback(ClientPair.Key) == EBreakBehavior::Break)
@@ -113,20 +113,29 @@ namespace UE::ConcertSyncServer::Replication
 			return EConcertSessionResponseCode::Success;
 		}
 
-		Clients.Emplace(ClientId, MakeShared<FConcertReplicationClient>(MoveTemp(StreamDescriptions), ClientId, Session, ReplicationCache));
+		Clients.Emplace(
+			ClientId,
+			MakeUnique<FConcertReplicationClient>(
+				MoveTemp(StreamDescriptions),
+				ClientId,
+				Session,
+				ReplicationCache,
+				ConcertSyncCore::FGetObjectFrequencySettings::CreateRaw(this, &FConcertServerReplicationManager::GetObjectFrequencySettings)
+			)
+		);
 		Response = { EJoinReplicationErrorCode::Success };
 		return EConcertSessionResponseCode::Success;
 	}
 
 	EConcertSessionResponseCode FConcertServerReplicationManager::HandleQueryReplicationInfoRequest(
-		const FConcertSessionContext& ConcertSessionContext,
+		const FConcertSessionContext&,
 		const FConcertReplication_QueryReplicationInfo_Request& Request,
 		FConcertReplication_QueryReplicationInfo_Response& Response
 		)
 	{
 		for (const FGuid& EndpointId : Request.ClientEndpointIds)
 		{
-			const TSharedRef<FConcertReplicationClient>* Client = Clients.Find(EndpointId);
+			const TUniquePtr<FConcertReplicationClient>* Client = Clients.Find(EndpointId);
 			if (!Client)
 			{
 				// This could happen if the client left the replication session before this request was answered
@@ -136,43 +145,38 @@ namespace UE::ConcertSyncServer::Replication
 			FReplicationClientQueriedInfo& EndpointInfo = Response.ClientInfo.Add(EndpointId);
 			if (!EnumHasAnyFlags(Request.QueryFlags, EConcertQueryClientStreamFlags::SkipStreamInfo))
 			{
-				const bool bSkipProperties = EnumHasAnyFlags(Request.QueryFlags, EConcertQueryClientStreamFlags::SkipProperties);
-				EndpointInfo.Streams = BuildClientStreamInfo(Client->Get(), bSkipProperties);
+				EndpointInfo.Streams = BuildClientStreamInfo(*Client->Get(), Request.QueryFlags);
 			}
 			if (!EnumHasAnyFlags(Request.QueryFlags, EConcertQueryClientStreamFlags::SkipAuthority))
 			{
-				EndpointInfo.Authority = BuildClientAuthorityInfo(Client->Get());
+				EndpointInfo.Authority = BuildClientAuthorityInfo(*Client->Get());
 			}
+			
 		}
 
 		Response.ErrorCode = EReplicationResponseErrorCode::Handled;
 		return EConcertSessionResponseCode::Success;
 	}
 
-	TArray<FSharedReplicationStreamDescription> FConcertServerReplicationManager::BuildClientStreamInfo(const FConcertReplicationClient& Client, bool bSkipProperties) const
+	TArray<FSharedReplicationStreamDescription> FConcertServerReplicationManager::BuildClientStreamInfo(const FConcertReplicationClient& Client, EConcertQueryClientStreamFlags QueryFlags)
 	{
 		TArray<FSharedReplicationStreamDescription> Result;
-		Algo::Transform(Client.GetStreamDescriptions(), Result, [bSkipProperties](const FReplicationStreamDescription& Description)
+		Algo::Transform(Client.GetStreamDescriptions(), Result, [QueryFlags](const FReplicationStreamDescription& Description)
 		{
-			if (!bSkipProperties)
+			using namespace ConcertSyncCore;
+			
+			EReplicationStreamCloneFlags Flags = EReplicationStreamCloneFlags::None;
+			if (EnumHasAnyFlags(QueryFlags, EConcertQueryClientStreamFlags::SkipProperties))
 			{
-				return Description.BaseDescription;
+				Flags |= EReplicationStreamCloneFlags::SkipProperties;
 			}
-
-			const FSharedReplicationStreamDescription& CopiedDescription = Description.BaseDescription;
-			const FObjectReplicationMap& CopiedReplicationMap = CopiedDescription.ReplicationMap;
-
-			// It was requested to skip sending the properties (saves network bandwidth)
-			FSharedReplicationStreamDescription StrippedResult;
-			StrippedResult.Identifier = CopiedDescription.Identifier;
-			StrippedResult.ReplicationMap.ReplicatedObjects.Reserve(CopiedReplicationMap.ReplicatedObjects.Num());
-			for (const TPair<FSoftObjectPath, FReplicatedObjectInfo>& ObjectInfo : CopiedReplicationMap.ReplicatedObjects)
+			if (EnumHasAnyFlags(QueryFlags, EConcertQueryClientStreamFlags::SkipFrequency))
 			{
-				StrippedResult.ReplicationMap.ReplicatedObjects.Add(ObjectInfo.Key, { ObjectInfo.Value.ClassPath });
+				Flags |= EReplicationStreamCloneFlags::SkipFrequency;
 			}
-			return StrippedResult;
+			
+			return Description.BaseDescription.Clone(Flags);
 		});
-
 		return Result;
 	}
 
@@ -230,13 +234,34 @@ namespace UE::ConcertSyncServer::Replication
 
 	void FConcertServerReplicationManager::Tick(IConcertServerSession& InSession, float InDeltaTime)
 	{
-		// TODO: Time slice replication clients
-		for (TPair<FGuid, TSharedRef<FConcertReplicationClient>>& Client : Clients)
+		// TODO UE-190714: We may want to set a time budget for clients.
+		// TODO UE-203340: Time slice clients
+		for (TPair<FGuid, TUniquePtr<FConcertReplicationClient>>& Client : Clients)
 		{
-			// TODO: Load time budget from config
-			constexpr float TimeBudget = 1.f / 60.f;
-			Client.Value->ProcessClient(TimeBudget);
+			Client.Value->ProcessClient({ InDeltaTime });
 		}
 	}
+
+	FConcertObjectReplicationSettings FConcertServerReplicationManager::GetObjectFrequencySettings(const FReplicatedObjectId& Object) const
+	{
+		const TUniquePtr<FConcertReplicationClient>* Client = Clients.Find(Object.SenderEndpointId);
+		if (!ensureMsgf(Client, TEXT("Caller is trying to retrieve non-existing client")))
+		{
+			UE_LOG(LogConcert, Warning, TEXT("Requested frequency settings for unknown client %s"), *Object.SenderEndpointId.ToString());
+			return {};
+		}
+		
+		const FReplicationStreamDescription* Stream = Client->Get()->GetStreamDescriptions().FindByPredicate([&Object](const FReplicationStreamDescription& Description)
+			{
+				return Description.BaseDescription.Identifier == Object.StreamId;
+			});
+		if (!ensureMsgf(Stream, TEXT("Caller is trying to retrieve an object that is not registered with the client")))
+		{
+			UE_LOG(LogConcert, Warning, TEXT("Requested frequency settings for unknown stream %s and object %s"), *Object.StreamId.ToString(), *Object.Object.ToString());
+			return {};
+		}
+			
+		return Stream->BaseDescription.FrequencySettings.GetSettingsFor(Object.Object);
+	}
 }
 
diff --git a/Engine/Plugins/Developer/Concert/ConcertSync/ConcertSyncServer/Source/ConcertSyncServer/Private/Replication/ConcertServerReplicationManager.h b/Engine/Plugins/Developer/Concert/ConcertSync/ConcertSyncServer/Source/ConcertSyncServer/Private/Replication/ConcertServerReplicationManager.h
--- a/Engine/Plugins/Developer/Concert/ConcertSync/ConcertSyncServer/Source/ConcertSyncServer/Private/Replication/ConcertServerReplicationManager.h
+++ b/Engine/Plugins/Developer/Concert/ConcertSync/ConcertSyncServer/Source/ConcertSyncServer/Private/Replication/ConcertServerReplicationManager.h
@@ -16,6 +16,8 @@
 class IConcertClientReplicationBridge;
 class IConcertServerSession;
 
+enum class EConcertQueryClientStreamFlags : uint8;
+
 struct FConcertReplication_ChangeStream_Response;
 struct FConcertReplication_QueryReplicationInfo_Response;
 struct FConcertReplication_QueryReplicationInfo_Request;
@@ -70,8 +72,11 @@ namespace UE::ConcertSyncServer::Replication
 		/** Receives replication events from all endpoints. */
 		FServerObjectReplicationReceiver ReplicationDataReceiver;
 
-		/** Clients that have requested to join replication. Maps client ID to replication info. */
-		TMap<FGuid, TSharedRef<FConcertReplicationClient>> Clients;
+		/**
+		 * Clients that have requested to join replication. Maps client ID to replication info.
+		 * Clients are stored in a unique ptr to avoid dealing with reallocations when the map resizes.
+		 */
+		TMap<FGuid, TUniquePtr<FConcertReplicationClient>> Clients;
 
 		// Joining
 		EConcertSessionResponseCode HandleJoinReplicationSessionRequest(const FConcertSessionContext& ConcertSessionContext, const FConcertReplication_Join_Request& Request, FConcertReplication_Join_Response& Response);
@@ -80,7 +85,7 @@ namespace UE::ConcertSyncServer::Replication
 		// Querying
 		EConcertSessionResponseCode HandleQueryReplicationInfoRequest(const FConcertSessionContext& ConcertSessionContext, const FConcertReplication_QueryReplicationInfo_Request& Request, FConcertReplication_QueryReplicationInfo_Response& Response);
 		/** Gets all registered streams and optionally removes the properties. */
-		TArray<FSharedReplicationStreamDescription> BuildClientStreamInfo(const FConcertReplicationClient& Client, bool bSkipProperties) const;
+		static TArray<FSharedReplicationStreamDescription> BuildClientStreamInfo(const FConcertReplicationClient& Client, EConcertQueryClientStreamFlags QueryFlags);
 		/** Maps the client's streams to the objects in that stream the client has taken authority over. */
 		TArray<FReplicationAuthorityInfo> BuildClientAuthorityInfo(const FConcertReplicationClient& Client) const;
 
@@ -98,5 +103,8 @@ namespace UE::ConcertSyncServer::Replication
 		 * This is configured in an .ini. TODO: Add config
 		 */
 		void Tick(IConcertServerSession& InSession, float InDeltaTime);
+		
+		/** Callback to clients for obtaining an object's frequency settings. */
+		FConcertObjectReplicationSettings GetObjectFrequencySettings(const FReplicatedObjectId& Object) const;
 	};
 }
\ No newline at end of file
diff --git a/Engine/Plugins/Developer/Concert/ConcertSync/ConcertSyncServer/Source/ConcertSyncServer/Private/Replication/ConcertServerReplicationManager_ChangeStream.cpp b/Engine/Plugins/Developer/Concert/ConcertSync/ConcertSyncServer/Source/ConcertSyncServer/Private/Replication/ConcertServerReplicationManager_ChangeStream.cpp
--- a/Engine/Plugins/Developer/Concert/ConcertSync/ConcertSyncServer/Source/ConcertSyncServer/Private/Replication/ConcertServerReplicationManager_ChangeStream.cpp
+++ b/Engine/Plugins/Developer/Concert/ConcertSync/ConcertSyncServer/Source/ConcertSyncServer/Private/Replication/ConcertServerReplicationManager_ChangeStream.cpp
@@ -133,13 +133,21 @@ namespace UE::ConcertSyncServer::Replication
 		}
 
 		/** Checks whether this request is valid to apply. */
-		static bool ShouldAcceptRequest(const FConcertReplication_ChangeStream_Request& Request, const FConcertReplicationClient& Client, const FAuthorityManager& AuthorityManager, FConcertReplication_ChangeStream_Response& OutResponse)
+		static bool ShouldAcceptRequest(
+			const FConcertReplication_ChangeStream_Request& Request,
+			const FConcertReplicationClient& Client,
+			const FAuthorityManager& AuthorityManager,
+			FConcertReplication_ChangeStream_Response& OutResponse
+			)
 		{
 			OutResponse.ErrorCode = EReplicationResponseErrorCode::Handled;
+			
 			ValidatePutObjectsRequestSemantics(Request, Client, OutResponse);
 			ValidateAddedStreamsAreUnique(Request, Client, OutResponse);
 			LookForAuthorityConflicts(Request, Client, AuthorityManager, OutResponse);
 			ValidateUnpacking(Request, OutResponse);
+			ConcertSyncCore::Replication::ChangeStreamUtils::ValidateFrequencyChanges(Request, Client.GetStreamDescriptions(), &OutResponse.FrequencyErrors);
+			
 			return OutResponse.IsSuccess();
 		}
 	}
@@ -152,18 +160,18 @@ namespace UE::ConcertSyncServer::Replication
 		Response = {};
 		
 		const FGuid SendingClientId = ConcertSessionContext.SourceEndpointId;
-		const TSharedRef<FConcertReplicationClient>* SendingClient = Clients.Find(SendingClientId);
-		if (SendingClient && Private::ShouldAcceptRequest(Request, SendingClient->Get(), AuthorityManager.Get(), Response))
+		const TUniquePtr<FConcertReplicationClient>* SendingClient = Clients.Find(SendingClientId);
+		if (SendingClient && Private::ShouldAcceptRequest(Request, *SendingClient->Get(), AuthorityManager.Get(), Response))
 		{
 			// If the client had authority over any objects that were removed by this request, authority must be cleaned up
-			ConcertSyncCore::Replication::ChangeStreamUtils::ForEachObjectLosingAuthority(Request, SendingClient->Get().GetStreamDescriptions(),
+			ConcertSyncCore::Replication::ChangeStreamUtils::ForEachObjectLosingAuthority(Request, SendingClient->Get()->GetStreamDescriptions(),
 				[this, &SendingClientId](const FObjectInStreamID& RemovedObject)
 				{
 					AuthorityManager->RemoveAuthority({ RemovedObject, SendingClientId});
 					return EBreakBehavior::Continue;
 				});
 			
-			SendingClient->Get().ApplyValidatedRequest(Request);
+			SendingClient->Get()->ApplyValidatedRequest(Request);
 		}
 		else
 		{
diff --git a/Engine/Plugins/Developer/Concert/ConcertSync/ConcertSyncTest/Source/ConcertSyncTest/Private/Replication/Components/FrequencyTrackerTests.cpp b/Engine/Plugins/Developer/Concert/ConcertSync/ConcertSyncTest/Source/ConcertSyncTest/Private/Replication/Components/FrequencyTrackerTests.cpp
new file mode 100644
--- /dev/null
+++ b/Engine/Plugins/Developer/Concert/ConcertSync/ConcertSyncTest/Source/ConcertSyncTest/Private/Replication/Components/FrequencyTrackerTests.cpp
@@ -0,0 +1,84 @@
+// Copyright Epic Games, Inc. All Rights Reserved.
+
+#include "Replication/Data/ReplicationFrequencySettings.h"
+#include "Replication/Processing/Proxy/FrequencyTracker.h"
+
+#include "Misc/AutomationTest.h"
+#include "Misc/Paths.h"
+#include "HAL/FileManager.h"
+#include "HAL/PlatformTime.h"
+
+namespace UE::ConcertSyncTests
+{
+	/**
+	 * Tests that FFrequencyTracker works as expected:
+	 *	- Tracks objects correctly,
+	 *	- Does not leak objects when cleaning cache
+	 */
+	IMPLEMENT_SIMPLE_AUTOMATION_TEST(FFrequencyTrackerTest, "Editor.Concert.Replication.Components.FrequencyTracker", EAutomationTestFlags::EditorContext | EAutomationTestFlags::EngineFilter);
+	bool FFrequencyTrackerTest::RunTest(const FString& Parameters)
+	{
+		using namespace ConcertSyncCore;
+		class FFrequencyTrackerMock : public FFrequencyTracker_CleanByTick
+		{
+		public:
+
+			FFrequencyTrackerMock(const double CacheRetentionSeconds)
+				: FFrequencyTracker_CleanByTick(CacheRetentionSeconds)
+			{}
+			
+			using FFrequencyTracker_CleanByTick::GetObjectsToLastProcessed;
+		};
+
+		// 1. Init test data
+		const FSoftObjectPath ReplicatedObject = GetMutableDefault<UObject>();
+		const FObjectInStreamID SpecifiedRateObject { FGuid{ 1, 0, 0, 0 }, ReplicatedObject };
+		const FObjectInStreamID RealtimeObject { FGuid{ 2, 0, 0, 0 }, ReplicatedObject };
+
+		constexpr FConcertObjectReplicationSettings SpecifiedRateSettings { EConcertObjectReplicationMode::SpecifiedRate, 30 };
+		constexpr FConcertObjectReplicationSettings RealtimeSettings { EConcertObjectReplicationMode::Realtime };
+
+		const double Now = FPlatformTime::Seconds();
+		constexpr double CacheRetentionSeconds = 5.f;
+		FFrequencyTrackerMock TestFrequencyTracker(CacheRetentionSeconds);
+
+		// 2.1 Simulate replicating immediately
+		TestEqual(TEXT("Can immediately replicate realtime object"),
+			TestFrequencyTracker.TrackAndCheckObject(RealtimeObject, RealtimeSettings, Now),
+			ECanProcessObject::CanProcess
+			);
+		TestEqual(TEXT("Can immediately replicate throttled object"),
+			TestFrequencyTracker.TrackAndCheckObject(SpecifiedRateObject, SpecifiedRateSettings, Now),
+			ECanProcessObject::CanProcess
+			);
+		TestFrequencyTracker.OnProcessObject(RealtimeObject, Now);
+		TestFrequencyTracker.OnProcessObject(SpecifiedRateObject, Now);
+		
+		TestEqual(TEXT("Throttled needs to wait after sending"),
+			TestFrequencyTracker.TrackAndCheckObject(SpecifiedRateObject, SpecifiedRateSettings, Now),
+			ECanProcessObject::TooEarly
+			);
+
+		// 2.2 Simulate replicating after specified rate has elapsed
+		const double AfterIntervalElapsed = Now + SpecifiedRateSettings.GetUpdateIntervalInSeconds();
+		TestEqual(TEXT("Can still replicate realtime object"),
+			TestFrequencyTracker.TrackAndCheckObject(RealtimeObject, RealtimeSettings, AfterIntervalElapsed),
+			ECanProcessObject::CanProcess
+			);
+		TestEqual(TEXT("Throttled object can send again"),
+			TestFrequencyTracker.TrackAndCheckObject(SpecifiedRateObject, SpecifiedRateSettings, AfterIntervalElapsed),
+			ECanProcessObject::CanProcess
+			);
+
+		// 2.3 Clean up cache
+		constexpr double FakeWaitSeconds = CacheRetentionSeconds + 1.0;
+		// Pretend that FakeWaitTime has elapsed so the cache is cleared.
+		TestFrequencyTracker.Tick(FakeWaitSeconds, AfterIntervalElapsed + FakeWaitSeconds);
+		TestEqual(TEXT("Cache was cleaned up after long time of not sending"),
+			TestFrequencyTracker.GetObjectsToLastProcessed().Num(),
+			0
+			);
+
+		return true;
+	}
+}
diff --git a/Engine/Plugins/Developer/Concert/ConcertSync/ConcertSyncTest/Source/ConcertSyncTest/Private/Replication/FrequencyTests.cpp b/Engine/Plugins/Developer/Concert/ConcertSync/ConcertSyncTest/Source/ConcertSyncTest/Private/Replication/FrequencyTests.cpp
new file mode 100644
--- /dev/null
+++ b/Engine/Plugins/Developer/Concert/ConcertSync/ConcertSyncTest/Source/ConcertSyncTest/Private/Replication/FrequencyTests.cpp
@@ -0,0 +1,153 @@
+// Copyright Epic Games, Inc. All Rights Reserved.
+
+#include "Replication/Data/ReplicationStreamDescription.h"
+#include "Replication/IConcertClientReplicationManager.h"
+#include "Replication/Util/SendReceiveObjectTestBase.h"
+#include "TestReflectionObject.h"
+
+#include "Misc/AutomationTest.h"
+#include "Tests/AutomationCommon.h"
+
+namespace UE::ConcertSyncTests::Replication::Frequency
+{
+	/** Makes the test object replicate at a frequency of 30 times per second. */
+	class FSendReceiveObjectWithFrequencyTest : public FSendReceiveObjectTestBase
+	{
+	public:
+		FSendReceiveObjectWithFrequencyTest(const FString& InName, const bool bInComplexTask)
+			: FSendReceiveObjectTestBase(InName, bInComplexTask)
+		{}
+
+		static constexpr uint8 GetUpdateRate() { return 30; } 
+		static constexpr double GetUpdateInterval() { return 1.0 / static_cast<double>(GetUpdateRate()); } 
+
+		//~ Begin FSendReceiveTestBase Interface
+		virtual ConcertSyncClient::Replication::FJoinReplicatedSessionArgs CreateSenderArgs() override
+		{
+			// CreateHandshakeArgsFrom defaults to EConcertObjectReplicationMode::Realtime which must be overriden for this test.
+			constexpr EConcertObjectReplicationMode ReplicationMode = EConcertObjectReplicationMode::SpecifiedRate;
+			return CreateHandshakeArgsFrom(*TestObject, SenderStreamId, ReplicationMode, GetUpdateRate());
+		}
+		//~ End FSendReceiveTestBase Interface
+	};
+	
+	DEFINE_LATENT_AUTOMATION_COMMAND_FOUR_PARAMETER(FSendAndCountReceivedCommand, FString, Name, FSendReceiveObjectTestBase*, Test, uint32, ExpectedNumServerToReceive, uint32, ExpectedNumClientToReceive);
+	bool FSendAndCountReceivedCommand::Update()
+	{
+		uint32 NumServerReceived = 0;
+		uint32 NumClientReceived = 0;
+		auto OnServerReceive = [&NumServerReceived](const FConcertSessionContext& Context, const FConcertReplication_BatchReplicationEvent& Event)
+		{
+			++NumServerReceived;
+		};
+		auto OnClientReceive = [&NumClientReceived](const FConcertSessionContext& Context, const FConcertReplication_BatchReplicationEvent& Event)
+		{
+			++NumClientReceived;
+		};
+		Test->SimulateSendObjectToReceiver(OnServerReceive, OnClientReceive);
+		
+		Test->TestEqual(FString::Printf(TEXT("%s: Server received data"), *Name), NumServerReceived, ExpectedNumServerToReceive);
+		Test->TestEqual(FString::Printf(TEXT("%s: Client received data"), *Name), NumClientReceived, ExpectedNumClientToReceive);
+		return true;
+	}
+	
+	/** Tests that something is sent after the replication interval has expired. */
+	IMPLEMENT_CUSTOM_SIMPLE_AUTOMATION_TEST(FReceiveAfterSpecifiedFrequencyInterval, FSendReceiveObjectWithFrequencyTest, "Editor.Concert.Replication.Frequency.SendingClientRespectsFrequency", EAutomationTestFlags::EditorContext | EAutomationTestFlags::EngineFilter);
+	bool FReceiveAfterSpecifiedFrequencyInterval::RunTest(const FString& Parameters)
+	{
+		SetUpClientAndServer();
+		ClientReplicationManager_Sender->TakeAuthorityOver({ TestObject });
+
+		// 1st frame replicates ...
+		ADD_LATENT_AUTOMATION_COMMAND(FSendAndCountReceivedCommand(TEXT("Frame 1 sends"), this, 1, 1));
+		// ... subsequent frames must wait
+		ADD_LATENT_AUTOMATION_COMMAND(FSendAndCountReceivedCommand(TEXT("Frame 2 does not send"), this, 0, 0));
+		ADD_LATENT_AUTOMATION_COMMAND(FSendAndCountReceivedCommand(TEXT("Frame 3 does not send"), this, 0, 0));
+		ADD_LATENT_AUTOMATION_COMMAND(FSendAndCountReceivedCommand(TEXT("Frame 4 does not send"), this, 0, 0));
+		ADD_LATENT_AUTOMATION_COMMAND(FSendAndCountReceivedCommand(TEXT("Frame 5 does not send"), this, 0, 0));
+
+		ADD_LATENT_AUTOMATION_COMMAND(FWaitLatentCommand(GetUpdateInterval()));
+		// ... until GetUpdateInterval() has elapsed
+		ADD_LATENT_AUTOMATION_COMMAND(FSendAndCountReceivedCommand(TEXT("Later frame sends again"), this, 1, 1));
+		return true;
+	}
+
+	/** Tests that after about 1 second, the expected number of events was received. */
+	IMPLEMENT_CUSTOM_SIMPLE_AUTOMATION_TEST(FSimulateAt30FPS, FSendReceiveObjectWithFrequencyTest, "Editor.Concert.Replication.Frequency.SimulateAt30FPS", EAutomationTestFlags::EditorContext | EAutomationTestFlags::EngineFilter);
+	bool FSimulateAt30FPS::RunTest(const FString& Parameters)
+	{
+		// 1. Set up server
+		SetUpClientAndServer();
+		ClientReplicationManager_Sender->TakeAuthorityOver({ TestObject });
+
+		// 2. Simulate the engine ticking for 1 second
+		// This was tried with ADD_LATENT_AUTOMATION_COMMAND but the test results were too inconsistent and imprecise
+		uint32 NumServerReceived = 0;
+		uint32 NumClientReceived = 0;
+		bool bIsFirstServerReceive = true;
+		bool bIsFirstClientReceive = true;
+		
+		double CurrentSeconds = FPlatformTime::Seconds();
+		double LastServerReceiveTime = CurrentSeconds;
+		double LastClientReceiveTime = CurrentSeconds;
+		
+		constexpr double UpdateIntervalSeconds = GetUpdateInterval();
+		const double EndSeconds = CurrentSeconds + 1.0;
+		for (; CurrentSeconds < EndSeconds; CurrentSeconds = FPlatformTime::Seconds())
+		{
+			const double DeltaSecSinceLast_Server = CurrentSeconds - LastServerReceiveTime;
+			const double DeltaSecSinceLast_Client = CurrentSeconds - LastClientReceiveTime;
+			
+			bool bReceivedServer = false;
+			bool bReceivedClient = false;
+			auto OnServerReceive = [&](const FConcertSessionContext&, const FConcertReplication_BatchReplicationEvent&)
+			{
+				LastServerReceiveTime = CurrentSeconds;
+				++NumServerReceived;
+				bReceivedServer = true;
+			};
+			auto OnClientReceive = [&](const FConcertSessionContext&, const FConcertReplication_BatchReplicationEvent&)
+			{
+				LastClientReceiveTime = CurrentSeconds;
+				++NumClientReceived;
+				bReceivedClient = true;
+			};
+			// This simulates an engine tick
+			SimulateSendObjectToReceiver(OnServerReceive, OnClientReceive);
+
+			// 2.1 Verify that every received update is in acceptable percentage range
+			constexpr double AcceptableTimeRange = UpdateIntervalSeconds * 0.01;
+			constexpr double MinDeltaTime = UpdateIntervalSeconds - AcceptableTimeRange;
+			constexpr double MaxDeltaTime = UpdateIntervalSeconds + AcceptableTimeRange;
+			if (bReceivedServer)
+			{
+				const bool bMin = MinDeltaTime <= DeltaSecSinceLast_Server || bIsFirstServerReceive;
+				const bool bMax = DeltaSecSinceLast_Server <= MaxDeltaTime || bIsFirstServerReceive;
+				AddErrorIfFalse(bMin, FString::Printf(TEXT("Server: MinDeltaTime (%f) <= DeltaSecSinceLast_Server (%f)"), MinDeltaTime, DeltaSecSinceLast_Server));
+				AddErrorIfFalse(bMax, FString::Printf(TEXT("Server: DeltaSecSinceLast_Server (%f) <= MaxDeltaTime (%f)"), DeltaSecSinceLast_Server, MaxDeltaTime));
+				bIsFirstServerReceive = false;
+			}
+			if (bReceivedClient)
+			{
+				const bool bMin = MinDeltaTime <= DeltaSecSinceLast_Client || bIsFirstClientReceive;
+				const bool bMax = DeltaSecSinceLast_Client <= MaxDeltaTime || bIsFirstClientReceive;
+				AddErrorIfFalse(bMin, FString::Printf(TEXT("Client: MinDeltaTime (%f) <= DeltaSecSinceLast_Client (%f)"), MinDeltaTime, DeltaSecSinceLast_Client));
+				AddErrorIfFalse(bMax, FString::Printf(TEXT("Client: DeltaSecSinceLast_Client (%f) <= MaxDeltaTime (%f)"), DeltaSecSinceLast_Client, MaxDeltaTime));
+				bIsFirstClientReceive = false;
+			}
+		}
+
+		// 2.2 Verify that the expected number of events was received in the past 1 second
+		// We may be off by one event due to double rounding errors
+		constexpr uint8 AcceptableDropAmount = 1;
+		constexpr uint8 MinExpected = GetUpdateRate() - AcceptableDropAmount;
+		constexpr uint8 MaxExpected = GetUpdateRate();
+		AddErrorIfFalse(MinExpected <= NumServerReceived, FString::Printf(TEXT("Server: MinExpected (%d)<= NumServerReceived (%d)"), MinExpected, NumServerReceived));
+		AddErrorIfFalse(NumServerReceived <= MaxExpected, FString::Printf(TEXT("Server: NumServerReceived (%d) <= MaxExpected (%d)"), NumServerReceived, MaxExpected));
+		AddErrorIfFalse(MinExpected <= NumClientReceived, FString::Printf(TEXT("Server: MinExpected (%d)<= NumClientReceived (%d)"), MinExpected, NumClientReceived));
+		AddErrorIfFalse(NumClientReceived <= MaxExpected, FString::Printf(TEXT("Server: NumClientReceived (%d) <= MaxExpected (%d)"), NumClientReceived, MaxExpected));
+
+		// 2.1 (time between updates) and 2.2 (number of updates) correctly test that we replicate at the expected frequency.
+		return true;
+	}
+}
diff --git a/Engine/Plugins/Developer/Concert/ConcertSync/ConcertSyncTest/Source/ConcertSyncTest/Private/Replication/ReplicationSendReceiveTests.cpp b/Engine/Plugins/Developer/Concert/ConcertSync/ConcertSyncTest/Source/ConcertSyncTest/Private/Replication/ReplicationSendReceiveTests.cpp
--- a/Engine/Plugins/Developer/Concert/ConcertSync/ConcertSyncTest/Source/ConcertSyncTest/Private/Replication/ReplicationSendReceiveTests.cpp
+++ b/Engine/Plugins/Developer/Concert/ConcertSync/ConcertSyncTest/Source/ConcertSyncTest/Private/Replication/ReplicationSendReceiveTests.cpp
@@ -117,8 +117,8 @@ namespace UE::ConcertSyncTests::Replication::SendReceiveFlow
 		{
 			ConcertSyncClient::Replication::FJoinReplicatedSessionArgs SenderJoinArgs;
 			const UClass& TestClass = *TestObject->GetClass();
-			
-			FReplicatedObjectInfo FloatProperties { &TestClass };
+
+			FReplicatedObjectInfo FloatProperties {&TestClass};
 			FloatProperties.PropertySelection.ReplicatedProperties.Add(*FConcertPropertyChain::CreateFromPath(TestClass, { GET_MEMBER_NAME_CHECKED(UTestReflectionObject, Float) }));
 			FReplicatedObjectInfo VectorProperties { &TestClass };
 			VectorProperties.PropertySelection.ReplicatedProperties.Add(*FConcertPropertyChain::CreateFromPath(TestClass, { GET_MEMBER_NAME_CHECKED(UTestReflectionObject, Vector) }));
@@ -126,9 +126,14 @@ namespace UE::ConcertSyncTests::Replication::SendReceiveFlow
 			FReplicationStreamDescription FloatStream;
 			FloatStream.BaseDescription.Identifier = FloatStreamId;
 			FloatStream.BaseDescription.ReplicationMap.ReplicatedObjects.Add(TestObject, FloatProperties);
+			// Use realtime replication, otherwise we'll have to pass time for the frequency system
+			FloatStream.BaseDescription.FrequencySettings.Defaults.ReplicationMode = EConcertObjectReplicationMode::Realtime;
+			
 			FReplicationStreamDescription VectorStream;
 			VectorStream.BaseDescription.Identifier = VectorStreamId;
 			VectorStream.BaseDescription.ReplicationMap.ReplicatedObjects.Add(TestObject, VectorProperties);
+			// Use realtime replication, otherwise we'll have to pass time for the frequency system
+			VectorStream.BaseDescription.FrequencySettings.Defaults.ReplicationMode = EConcertObjectReplicationMode::Realtime;
 			
 			SenderJoinArgs.Streams.Add(FloatStream);
 			SenderJoinArgs.Streams.Add(VectorStream);
diff --git a/Engine/Plugins/Developer/Concert/ConcertSync/ConcertSyncTest/Source/ConcertSyncTest/Private/Replication/ReplicationStreamRequestTests.cpp b/Engine/Plugins/Developer/Concert/ConcertSync/ConcertSyncTest/Source/ConcertSyncTest/Private/Replication/ReplicationStreamRequestTests.cpp
deleted file mode 100644
diff --git a/Engine/Plugins/Developer/Concert/ConcertSync/ConcertSyncTest/Source/ConcertSyncTest/Private/Replication/StreamRequestTests_Change.cpp b/Engine/Plugins/Developer/Concert/ConcertSync/ConcertSyncTest/Source/ConcertSyncTest/Private/Replication/StreamRequestTests_Change.cpp
new file mode 100644
--- /dev/null
+++ b/Engine/Plugins/Developer/Concert/ConcertSync/ConcertSyncTest/Source/ConcertSyncTest/Private/Replication/StreamRequestTests_Change.cpp
@@ -0,0 +1,297 @@
+// Copyright Epic Games, Inc. All Rights Reserved.
+
+#include "Util/ClientServerCommunicationTest.h"
+
+#include "Replication/Data/ReplicationStreamDescription.h"
+#include "Replication/IConcertClientReplicationManager.h"
+#include "TestReflectionObject.h"
+#include "Util/ChangeStreamsTestBase.h"
+#include "Util/SendReceiveObjectTestBase.h"
+
+#include "Misc/AutomationTest.h"
+#include "UObject/UObjectGlobals.h"
+#include "UObject/Package.h"
+
+namespace UE::ConcertSyncTests::Replication::Stream
+{
+	/** Simple case of FConcertChangeStream_Request where the requesting client has authority and thus cannot generate any authority conflicts. */
+	IMPLEMENT_CUSTOM_SIMPLE_AUTOMATION_TEST(FSimpleChangeStream, FChangeStreamsTestBase, "Editor.Concert.Replication.Stream.SimpleChangeStream", EAutomationTestFlags::EditorContext | EAutomationTestFlags::EngineFilter);
+	bool FSimpleChangeStream::RunTest(const FString& Parameters)
+	{
+		using namespace ConcertSyncClient::Replication;
+
+		// 1. Start up client with float property stream
+		UTestReflectionObject* TestObject = NewObject<UTestReflectionObject>(GetTransientPackage());
+		auto[InitialStreamId, InitialStream] = CreateFloatPropertyStream(*TestObject);
+		SenderArgs.Streams = { InitialStream };
+		SetUpClientAndServer();
+
+		// 2.1 Modify existing stream
+		FConcertPropertySelection& ChangedInitialProperties = GetPropertySelection(InitialStream, *TestObject);
+		AddVectorProperty(ChangedInitialProperties);
+		FConcertReplication_ChangeStream_Request ModifyRequest;
+		ModifyRequest.ObjectsToPut.Add({ InitialStreamId, TestObject }, { ChangedInitialProperties });
+		ChangeStreamForSenderClientAndValidate(TEXT("ModifyRequest"), ModifyRequest, { InitialStream.BaseDescription });
+
+		// 2.2 Add new stream
+		auto[DynamicStreamId, DynamicStream] = CreateFloatPropertyStream(*TestObject);
+		FConcertReplication_ChangeStream_Request DynamicCreationRequest;
+		DynamicCreationRequest.StreamsToAdd.Add(DynamicStream.Pack());
+		ChangeStreamForSenderClientAndValidate(TEXT("DynamicCreationRequest"), DynamicCreationRequest, { InitialStream.BaseDescription, DynamicStream.BaseDescription });
+
+		// 2.3 Remove stream
+		FConcertReplication_ChangeStream_Request RemoveStreamRequest;
+		RemoveStreamRequest.StreamsToRemove.Add(InitialStreamId);
+		ChangeStreamForSenderClientAndValidate(TEXT("RemoveStreamRequest"), RemoveStreamRequest, { DynamicStream.BaseDescription });
+
+		// 2.4 Remove object > removes stream implicitly
+		FConcertReplication_ChangeStream_Request RemoveObjectRequest;
+		RemoveObjectRequest.ObjectsToRemove.Add({ DynamicStreamId, TestObject });
+		ChangeStreamForSenderClientAndValidate(TEXT("RemoveObjectRequest"), RemoveObjectRequest, {});
+		
+		return true;
+	}
+
+	/** Case of FConcertChangeStream_Request in which requester has authority over an object and creates two streams which have authority over overlapping properties. This is allowed. */
+	IMPLEMENT_CUSTOM_SIMPLE_AUTOMATION_TEST(FChangeStreamWithOverlappingProperties, FChangeStreamsTestBase, "Editor.Concert.Replication.Stream.ChangeStreamWithOverlappingProperties", EAutomationTestFlags::EditorContext | EAutomationTestFlags::EngineFilter);
+	bool FChangeStreamWithOverlappingProperties::RunTest(const FString& Parameters)
+	{
+		using namespace ConcertSyncClient::Replication;
+
+		// 1. Start client without any streams
+		UTestReflectionObject* TestObject = NewObject<UTestReflectionObject>(GetTransientPackage());
+		SetUpClientAndServer();
+		
+		// 2.1 Add stream with float property
+		auto[FloatStreamID, FloatStream] = CreateFloatPropertyStream(*TestObject);
+		FConcertReplication_ChangeStream_Request CreateFloatStreamRequest;
+		CreateFloatStreamRequest.StreamsToAdd.Add(FloatStream.Pack());
+		ChangeStreamForSenderClientAndValidate(TEXT("CreateFloatStreamRequest"), CreateFloatStreamRequest, { FloatStream.BaseDescription });
+		// 2.2 Add a stream with vector property
+		auto[VectorFloatStreamID, VectorFloatStream] = CreateVectorPropertyStream(*TestObject);
+		FConcertReplication_ChangeStream_Request CreateVectorFloatStreamRequest; 
+		CreateVectorFloatStreamRequest.StreamsToAdd.Add(VectorFloatStream.Pack());
+		ChangeStreamForSenderClientAndValidate(TEXT("CreateVectorFloatStreamRequest"), CreateVectorFloatStreamRequest, { FloatStream.BaseDescription, VectorFloatStream.BaseDescription });
+
+		// 2.3 Take authority over both streams
+		bool bTookAuthority = false;
+		FConcertReplication_ChangeAuthority_Request TakeAuthorityRequest;
+		TakeAuthorityRequest.TakeAuthority.Add(TestObject, FConcertStreamArray{{ FloatStreamID, VectorFloatStreamID }});
+		ClientReplicationManager_Sender->RequestAuthorityChange({ TakeAuthorityRequest })
+			.Next([this, &bTookAuthority](FConcertReplication_ChangeAuthority_Response&& Response)
+			{
+				bTookAuthority = true;
+				TestTrue(TEXT("ErrorCode == Handled"), Response.ErrorCode == EReplicationResponseErrorCode::Handled);
+				TestTrue(TEXT("Authority request > Success"), Response.RejectedObjects.IsEmpty());
+			});
+		TestTrue(TEXT("Authority request > Received response"), bTookAuthority);
+
+		// 2.4 Appending the overlapping float property works
+		FConcertReplication_ChangeStream_Request AppendFloatRequest;
+		FConcertPropertySelection& NewSelection = GetPropertySelection(VectorFloatStream, *TestObject);
+		AddFloatProperty(NewSelection);
+		AppendFloatRequest.ObjectsToPut.Add(FObjectInStreamID{ VectorFloatStreamID, TestObject }, FConcertReplication_ChangeStream_PutObject{ NewSelection });
+		ChangeStreamForSenderClientAndValidate(TEXT("AppendFloatRequest"), AppendFloatRequest, { FloatStream.BaseDescription, VectorFloatStream.BaseDescription });
+		
+		return true;
+	}
+
+	/**
+	 * The requester and client have authority over separate properties on the same object.
+	 * The request would cause the requester's existing authority to overlap with the other client, which is rejected.
+	 */
+	IMPLEMENT_CUSTOM_SIMPLE_AUTOMATION_TEST(FRejectChangeStreamWithConflictingAuthority, FChangeStreamsTestBase, "Editor.Concert.Replication.Stream.RejectChangeStreamWithConflictingAuthority", EAutomationTestFlags::EditorContext | EAutomationTestFlags::EngineFilter);
+	bool FRejectChangeStreamWithConflictingAuthority::RunTest(const FString& Parameters)
+	{
+		using namespace ConcertSyncClient::Replication;
+
+		// 1. Start client without any streams
+		UTestReflectionObject* TestObject = NewObject<UTestReflectionObject>(GetTransientPackage());
+		SetUpClientAndServer();
+		
+		// 2.1 Add stream with float property
+		auto[SenderStreamID, SenderStream] = CreateFloatPropertyStream(*TestObject);
+		FConcertReplication_ChangeStream_Request CreateFloatStreamRequest;
+		CreateFloatStreamRequest.StreamsToAdd.Add(SenderStream.Pack());
+		ChangeStreamForSenderClientAndValidate(TEXT("CreateFloatStreamRequest"), CreateFloatStreamRequest, { SenderStream.BaseDescription });
+
+		// 2.2 Have Sender take authority over the float property
+		bool bSenderTookAuthority = false;
+		ClientReplicationManager_Sender->TakeAuthorityOver({ TestObject })
+			.Next([this, &bSenderTookAuthority](const FConcertReplication_ChangeAuthority_Response& Response) mutable
+			{
+				bSenderTookAuthority = true;
+				TestTrue(TEXT("ErrorCode == Handled"), Response.ErrorCode == EReplicationResponseErrorCode::Handled);
+				TestEqual(TEXT("Sender > No rejection taking authority"), Response.RejectedObjects.Num(), 0); 
+			});
+		TestTrue(TEXT("Sender > Received response taking authority"), bSenderTookAuthority);
+
+		// 2.3 Have Receiver create vector stream
+		auto[ReceiverStreamID, ReceiverStream] = CreateVectorPropertyStream(*TestObject);
+		FConcertReplication_ChangeStream_Request CreateVectorStreamRequest; 
+		CreateVectorStreamRequest.StreamsToAdd.Add(ReceiverStream.Pack());
+		bool bAddedVectorStream = false;
+		ClientReplicationManager_Receiver->ChangeStream(CreateVectorStreamRequest)
+			.Next([this, &bAddedVectorStream](FConcertReplication_ChangeStream_Response&& Response)
+			{
+				bAddedVectorStream = true;
+				TestTrue(TEXT("ErrorCode == Handled"), Response.ErrorCode == EReplicationResponseErrorCode::Handled);
+				TestTrue(TEXT("Vector Stream > Success"), Response.IsSuccess());
+			});
+		TestTrue(TEXT("Vector Stream> Received response"), bAddedVectorStream);
+		// For simplicity, we don't do any further validation that other clients can see the change
+		
+		// 2.4 Have Receiver take authority over vector stream
+		bool bReceiverTookAuthority = false;
+		ClientReplicationManager_Receiver->TakeAuthorityOver({ TestObject })
+			.Next([this, &bReceiverTookAuthority](const FConcertReplication_ChangeAuthority_Response& Response) mutable
+			{
+				bReceiverTookAuthority = true;
+				TestTrue(TEXT("ErrorCode == Handled"), Response.ErrorCode == EReplicationResponseErrorCode::Handled);
+				TestEqual(TEXT("Receiver > No rejection taking authority"), Response.RejectedObjects.Num(), 0); 
+			});
+		TestTrue(TEXT("Receiver > Received response taking authority"), bReceiverTookAuthority);
+
+		// 2.5 Receiver is rejected from modifying the stream to include the Float property because the Sender has authority over it
+		const FObjectInStreamID TestObjectInReceiverStreamId { ReceiverStreamID, TestObject };
+		FConcertReplication_ChangeStream_Request AddFloatToStreamRequest;
+		FConcertPropertySelection NewSelection = GetPropertySelection(ReceiverStream, *TestObject);
+		AddFloatProperty(NewSelection);
+		AddFloatToStreamRequest.ObjectsToPut.Add(TestObjectInReceiverStreamId, { NewSelection });
+		bool bReceivedResponseAddingFloat = false;
+		
+		// Server logs a warning when rejecting - avoid the test being marked with a warning.
+		AddExpectedError(TEXT("Rejecting ChangeStream"));
+		ClientReplicationManager_Receiver->ChangeStream(AddFloatToStreamRequest)
+			.Next([this, TestObject, SenderStreamID = SenderStreamID, &TestObjectInReceiverStreamId, &bReceivedResponseAddingFloat](FConcertReplication_ChangeStream_Response&& Response)
+			{
+				bReceivedResponseAddingFloat = true;
+				TestTrue(TEXT("ErrorCode == Handled"), Response.ErrorCode == EReplicationResponseErrorCode::Handled);
+				TestEqual(TEXT("Append Float > 1 conflict"), Response.AuthorityConflicts.Num(), 1);
+				if (const FReplicatedObjectId* ConflictingObject = Response.AuthorityConflicts.Find(TestObjectInReceiverStreamId))
+				{
+					TestEqual(TEXT("Append float > Conflict > Sender Stream correct"), ConflictingObject->StreamId, SenderStreamID);
+					TestEqual(TEXT("Append float > Conflict > Object correct"), ConflictingObject->Object, FSoftObjectPath(TestObject));
+					TestEqual(TEXT("Append float > Conflict > Endpoint ID correct"), ConflictingObject->SenderEndpointId, Client_Sender->ClientSessionMock->GetSessionClientEndpointId());
+				}
+				else
+				{
+					AddError(TEXT("Expected to find an authority object conflict for TestObject"));
+				}
+				TestTrue(TEXT("Append Float > Failure"), Response.IsFailure());
+			});
+		TestTrue(TEXT("Append Float > Received response"), bReceivedResponseAddingFloat);
+		
+		return true;  
+	}
+
+	/** Tests that no changes are made if a sub-step of FConcertChangeStream_Request causes an error. */
+	IMPLEMENT_CUSTOM_SIMPLE_AUTOMATION_TEST(FChangeStreamIsAtmoic, FChangeStreamsTestBase, "Editor.Concert.Replication.Stream.ChangeStreamIsAtmoic", EAutomationTestFlags::EditorContext | EAutomationTestFlags::EngineFilter);
+	bool FChangeStreamIsAtmoic::RunTest(const FString& Parameters)
+	{
+		using namespace ConcertSyncClient::Replication;
+
+		// 1. Start client without any streams
+		UTestReflectionObject* TestObject = NewObject<UTestReflectionObject>(GetTransientPackage());
+		SetUpClientAndServer();
+		
+		// 2.1 Add stream with float property
+		auto[StreamId, Stream] = CreateFloatPropertyStream(*TestObject);
+		FConcertReplication_ChangeStream_Request CreateStreamRequest;
+		CreateStreamRequest.StreamsToAdd.Add(Stream.Pack());
+		ChangeStreamForSenderClientAndValidate(TEXT("CreateStreamRequest"), CreateStreamRequest, { Stream.BaseDescription });
+
+		// 2.2 Make a request that will fail and check that no changes were made to the original stream
+		FConcertReplication_ChangeStream_Request InvalidRequest;
+		FConcertPropertySelection NewSelection = GetPropertySelection(Stream, *TestObject);
+		AddFloatProperty(NewSelection);
+		InvalidRequest.ObjectsToPut.Add(FObjectInStreamID{ StreamId, TestObject}, FConcertReplication_ChangeStream_PutObject{ NewSelection });
+		InvalidRequest.StreamsToAdd.Add(Stream.Pack()); // This will make it fail due to pre-existing stream ID
+
+		// Server logs a warning when rejecting - avoid the test being marked with a warning.
+		AddExpectedError(TEXT("Rejecting ChangeStream"));
+		
+		bool bReceivedChangeStreamResponse = false;
+		ClientReplicationManager_Sender->ChangeStream(InvalidRequest)
+			.Next([this, &bReceivedChangeStreamResponse](FConcertReplication_ChangeStream_Response&& Response)
+			{
+				bReceivedChangeStreamResponse = true;
+				TestTrue(TEXT("ErrorCode == Handled"), Response.ErrorCode == EReplicationResponseErrorCode::Handled);
+				TestTrue(TEXT("Invalid Request > Modify Stream > Failure"), Response.IsFailure());
+			});
+		
+		TestTrue(TEXT("Invalid Request > Modify Stream > Received response"), bReceivedChangeStreamResponse);
+		ValidateSenderClientStreams(TEXT("Invalid Request"), { Stream.BaseDescription });
+		
+		return true;
+	}
+
+	/** Removing an object from a stream over which the requester has authority, clears the authority the client had. Tests that another client can now claim authority over the previously authored properties. */
+	IMPLEMENT_CUSTOM_SIMPLE_AUTOMATION_TEST(FRemovingObjectClearsAuthority, FChangeStreamsTestBase, "Editor.Concert.Replication.Stream.RemovingObjectClearsAuthority", EAutomationTestFlags::EditorContext | EAutomationTestFlags::EngineFilter);
+	bool FRemovingObjectClearsAuthority::RunTest(const FString& Parameters)
+	{
+		using namespace ConcertSyncClient::Replication;
+
+		// 1. Start clients with equivalent streams
+		UTestReflectionObject* TestObject = NewObject<UTestReflectionObject>(GetTransientPackage());
+		auto[SenderStreamId, SenderStream] = CreateFloatPropertyStream(*TestObject);
+		auto[ReceiverStreamId, ReceiverStream] = CreateFloatPropertyStream(*TestObject);
+		SenderArgs.Streams = { SenderStream };
+		ReceiverArgs.Streams = { ReceiverStream };
+		SetUpClientAndServer();
+		
+		// 2.1 Sender takes authority
+		ClientReplicationManager_Sender->TakeAuthorityOver({ TestObject })
+			.Next([this](FConcertReplication_ChangeAuthority_Response&& Response)
+			{
+				TestTrue(TEXT("ErrorCode == Handled"), Response.ErrorCode == EReplicationResponseErrorCode::Handled);
+				TestTrue(TEXT("Sender > Success taking authority"), Response.RejectedObjects.IsEmpty());
+			});
+
+		// 2.2 Sender removes the object
+		FConcertReplication_ChangeStream_Request RemoveObjectRequest;
+		RemoveObjectRequest.ObjectsToRemove.Add({ SenderStreamId, TestObject });
+		ChangeStreamForSenderClientAndValidate(TEXT("Remove Object"), RemoveObjectRequest, {});
+
+		// 2.3 Receiver can now take authority since 2.2 implicitly removed authority
+		ClientReplicationManager_Receiver->TakeAuthorityOver({ TestObject })
+			.Next([this](FConcertReplication_ChangeAuthority_Response&& Response)
+			{
+				TestTrue(TEXT("ErrorCode == Handled"), Response.ErrorCode == EReplicationResponseErrorCode::Handled);
+				TestTrue(TEXT("Receiver > Success taking authority"), Response.RejectedObjects.IsEmpty());
+			});
+		
+		return true;
+	}
+
+	/**
+     * Tests that client updates its local cache of the server state when RequestAuthorityChange times out.
+     */
+    IMPLEMENT_CUSTOM_SIMPLE_AUTOMATION_TEST(FChangingStreamTimeoutRetainsServerState, FChangeStreamsTestBase, "Editor.Concert.Replication.Stream.ChangeStreamTimeoutRetainsServerState", EAutomationTestFlags::EditorContext | EAutomationTestFlags::EngineFilter);
+    bool FChangingStreamTimeoutRetainsServerState::RunTest(const FString& Parameters)
+    {
+    	// 1. Start up client with float property stream
+    	UTestReflectionObject* TestObject = NewObject<UTestReflectionObject>(GetTransientPackage());
+    	auto[InitialStreamId, InitialStream] = CreateFloatPropertyStream(*TestObject);
+    	SenderArgs.Streams = { InitialStream };
+    	SetUpClientAndServer();
+		ServerSession->SetTestFlags(EServerSessionTestingFlags::AllowRequestTimeouts);
+
+    	// 2. Timeout request
+    	ServerSession->UnregisterCustomRequestHandler<FConcertReplication_ChangeStream_Request>();
+    	FConcertReplication_ChangeStream_Request Request;
+    	Request.ObjectsToRemove.Add({ InitialStreamId, TestObject });
+    	ClientReplicationManager_Sender->ChangeStream(Request);
+
+    	// 3. Check that the local client's server prediction was reverted
+    	TArray<FReplicationStreamDescription> Streams = ClientReplicationManager_Sender->GetRegisteredStreams();
+    	if (Streams.IsEmpty())
+    	{
+    		AddError(TEXT("Stream change was not reverted"));
+    		return true;
+    	}
+    	TestTrue(TEXT("Timed out Stream change reverted"), Streams[0].BaseDescription.ReplicationMap.ReplicatedObjects.Contains(TestObject));
+    	
+    	return true;
+    }
+}
\ No newline at end of file
diff --git a/Engine/Plugins/Developer/Concert/ConcertSync/ConcertSyncTest/Source/ConcertSyncTest/Private/Replication/StreamRequestTests_Frequency.cpp b/Engine/Plugins/Developer/Concert/ConcertSync/ConcertSyncTest/Source/ConcertSyncTest/Private/Replication/StreamRequestTests_Frequency.cpp
new file mode 100644
--- /dev/null
+++ b/Engine/Plugins/Developer/Concert/ConcertSync/ConcertSyncTest/Source/ConcertSyncTest/Private/Replication/StreamRequestTests_Frequency.cpp
@@ -0,0 +1,377 @@
+// Copyright Epic Games, Inc. All Rights Reserved.
+
+#include "Replication/Data/ReplicationStreamDescription.h"
+#include "Replication/IConcertClientReplicationManager.h"
+#include "TestReflectionObject.h"
+#include "Util/ChangeStreamsTestBase.h"
+#include "Util/SendReceiveObjectTestBase.h"
+
+#include "Misc/AutomationTest.h"
+#include "UObject/UObjectGlobals.h"
+#include "UObject/Package.h"
+
+namespace UE::ConcertSyncTests::Replication::Stream
+{
+	/**
+	 * Tests that frequency settings
+	 * - can be changed using all possible options (put, add, remove)
+	 * - the operation order is put, then add, then remove
+	 * - querying frequency info works
+	 * - replication rate cannot be set to 0
+	 * - atomicity: failing to change replication frequency fails all other changes
+	 * - frequency can only be overriden if object is in object replication map
+	 * - frequency override is removed if object is removed from object replication map
+	 */
+	IMPLEMENT_CUSTOM_SIMPLE_AUTOMATION_TEST(FChangeFrequency, FChangeStreamsTestBase, "Editor.Concert.Replication.Stream.Frequency.ChangeFrequency", EAutomationTestFlags::EditorContext | EAutomationTestFlags::EngineFilter);
+	bool FChangeFrequency::RunTest(const FString& Parameters)
+	{
+		// 1. Start up client with float property stream
+		UTestReflectionObject* MainTestObject = NewObject<UTestReflectionObject>(GetTransientPackage());
+		UTestReflectionObject* SecondaryTestObject = NewObject<UTestReflectionObject>(GetTransientPackage());
+		auto[FloatStreamId, FloatStream] = CreateFloatPropertyStream(*MainTestObject);
+		SenderArgs.Streams = { FloatStream };
+		
+		SetUpClientAndServer();
+		// We'll intentionally be causing errors below that cause a warning to be logged.
+		AddExpectedError(TEXT("Rejecting ChangeStream"), EAutomationExpectedErrorFlags::Contains, 2);
+
+		FSharedReplicationStreamDescription ExpectedStream = FloatStream.BaseDescription;
+		
+		// Change defaults to be SpecifiedRate at 42 update rate.
+		{
+			FConcertReplication_ChangeStream_Request ChangeRequest;
+			FConcertReplication_ChangeStream_Frequency FrequencyChange { .NewDefaults = { EConcertObjectReplicationMode::SpecifiedRate, 42 }, .Flags = EConcertReplicationChangeFrequencyFlags::SetDefaults };
+			ChangeRequest.FrequencyChanges.Add(FloatStreamId, FrequencyChange);
+			
+			ExpectedStream.FrequencySettings.Defaults = { EConcertObjectReplicationMode::SpecifiedRate, 42 };
+			ChangeStreamForSenderClientAndValidate(TEXT("SpecifiedRate with 42"), ChangeRequest, { ExpectedStream });
+		}
+		
+		// Not specifying update flags does not update default replication state.
+		{
+			FConcertReplication_ChangeStream_Request ChangeRequest;
+			FConcertReplication_ChangeStream_Frequency FrequencyChange { .NewDefaults = { EConcertObjectReplicationMode::SpecifiedRate, 1 } };
+			ChangeRequest.FrequencyChanges.Add(FloatStreamId, FrequencyChange);
+			ChangeStreamForSenderClientAndValidate(TEXT("Do not update defaults without update flags"), ChangeRequest, { ExpectedStream });
+		}
+		// Cannot set SpecifiedRate at 0 update rate.
+		{
+			FConcertReplication_ChangeStream_Request ChangeRequest;
+			FConcertReplication_ChangeStream_Frequency FrequencyChange { .NewDefaults = { EConcertObjectReplicationMode::SpecifiedRate, 0 }, .Flags = EConcertReplicationChangeFrequencyFlags::SetDefaults };
+			ChangeRequest.FrequencyChanges.Add(FloatStreamId, FrequencyChange);
+			ChangeStreamForSenderClientAndValidate(
+				TEXT("Defaults > Rate 0"),
+				ChangeRequest,
+				{ ExpectedStream },
+				EChangeStreamValidation::TestWasHandled | EChangeStreamValidation::TestResponseFailed
+				)
+				.Next([this, FloatStreamId](FConcertReplication_ChangeStream_Response&& Response)
+				{
+					const EConcertChangeStreamFrequencyErrorCode* ErrorCode = Response.FrequencyErrors.DefaultFailures.Find(FloatStreamId);
+					TestTrue(TEXT("Defaults > Rate 0 > Rejection code"), ErrorCode && *ErrorCode == EConcertChangeStreamFrequencyErrorCode::InvalidReplicationRate);
+				});
+		}
+		
+		// Failing frequency update reverts other changes in request (atomicity)
+		{
+			FConcertReplication_ChangeStream_Request ChangeRequest;
+			
+			FConcertReplication_ChangeStream_PutObject PutObject;
+			AddFloatProperty(PutObject.Properties);
+			PutObject.ClassPath = SecondaryTestObject->GetClass();
+			ChangeRequest.ObjectsToPut.Add({ FloatStreamId, SecondaryTestObject }, PutObject);
+			
+			FConcertReplication_ChangeStream_Frequency FrequencyChange;
+			FrequencyChange.OverridesToAdd.Add(SecondaryTestObject, { EConcertObjectReplicationMode::SpecifiedRate, 0 });
+			ChangeRequest.FrequencyChanges.Add(FloatStreamId, FrequencyChange);
+			
+			ChangeStreamForSenderClientAndValidate(
+				TEXT("Failing frequency update reverts other changes in request (atomicity)"),
+				ChangeRequest,
+				{ ExpectedStream },
+				EChangeStreamValidation::TestWasHandled | EChangeStreamValidation::TestResponseFailed
+				);
+		}
+		
+		// Add object and change frequency in same request
+		{
+			FConcertReplication_ChangeStream_Request ChangeRequest;
+			
+			FConcertReplication_ChangeStream_PutObject PutObject;
+			AddFloatProperty(PutObject.Properties);
+			PutObject.ClassPath = SecondaryTestObject->GetClass();
+			ChangeRequest.ObjectsToPut.Add({ FloatStreamId, SecondaryTestObject }, PutObject);
+			
+			FConcertReplication_ChangeStream_Frequency FrequencyChange;
+			FrequencyChange.OverridesToAdd.Add(SecondaryTestObject, { EConcertObjectReplicationMode::SpecifiedRate, 101 });
+			ChangeRequest.FrequencyChanges.Add(FloatStreamId, FrequencyChange);
+
+			FReplicatedObjectInfo& ExpectedSecondary = ExpectedStream.ReplicationMap.ReplicatedObjects.Add(SecondaryTestObject);
+			AddFloatProperty(ExpectedSecondary.PropertySelection);
+			ExpectedSecondary.ClassPath = SecondaryTestObject->GetClass();
+			ExpectedStream.FrequencySettings.ObjectOverrides.Add(SecondaryTestObject, { EConcertObjectReplicationMode::SpecifiedRate, 101 });
+			ChangeStreamForSenderClientAndValidate(TEXT("Register object and set frequency"), ChangeRequest, { ExpectedStream });
+		}
+		
+		return true;
+	}
+
+	/**
+	 * Tests that requests for unregistered objects rejected:
+	 * - Cannot change frequency for unknown stream
+	 * - Cannot change frequency for object not registerd to stream
+	 */
+	IMPLEMENT_CUSTOM_SIMPLE_AUTOMATION_TEST(FPutAddAndRemoveFrequencyFlow, FChangeStreamsTestBase, "Editor.Concert.Replication.Stream.Frequency.PutAddAndRemoveFrequencyFlow", EAutomationTestFlags::EditorContext | EAutomationTestFlags::EngineFilter);
+	bool FPutAddAndRemoveFrequencyFlow::RunTest(const FString& Parameters)
+	{
+		UTestReflectionObject* MainTestObject = NewObject<UTestReflectionObject>(GetTransientPackage());
+		UTestReflectionObject* SecondaryTestObject = NewObject<UTestReflectionObject>(GetTransientPackage());
+		auto[FloatStreamId, FloatStream] = CreateFloatPropertyStream(*MainTestObject);
+		SenderArgs.Streams = { FloatStream };
+		SetUpClientAndServer();
+
+		FSharedReplicationStreamDescription ExpectedStream = FloatStream.BaseDescription;
+		// Put SpecifiedRate 21
+		{
+			FConcertReplication_ChangeStream_Request ChangeRequest;
+			FConcertReplication_ChangeStream_Frequency FrequencyChange;
+			FrequencyChange.OverridesToPut.Add(MainTestObject, { EConcertObjectReplicationMode::SpecifiedRate, 21 });
+			ChangeRequest.FrequencyChanges.Add(FloatStreamId, FrequencyChange);
+			
+			ExpectedStream.FrequencySettings.ObjectOverrides.Add(MainTestObject, { EConcertObjectReplicationMode::SpecifiedRate, 21 });
+			ChangeStreamForSenderClientAndValidate(TEXT("Put SpecifiedRate 21"), ChangeRequest, { ExpectedStream });
+		}
+		// Add SpecifiedRate 22
+		{
+			FConcertReplication_ChangeStream_Request ChangeRequest;
+			FConcertReplication_ChangeStream_Frequency FrequencyChange;
+			FrequencyChange.OverridesToAdd.Add(MainTestObject, { EConcertObjectReplicationMode::SpecifiedRate, 22 });
+			ChangeRequest.FrequencyChanges.Add(FloatStreamId, FrequencyChange);
+			
+			ExpectedStream.FrequencySettings.ObjectOverrides.Add(MainTestObject, { EConcertObjectReplicationMode::SpecifiedRate, 22 });
+			ChangeStreamForSenderClientAndValidate(TEXT("Add SpecifiedRate 22"), ChangeRequest, { ExpectedStream });
+		}
+		// Remove 
+		{
+			FConcertReplication_ChangeStream_Request ChangeRequest;
+			FConcertReplication_ChangeStream_Frequency FrequencyChange;
+			FrequencyChange.OverridesToRemove.Add(MainTestObject);
+			ChangeRequest.FrequencyChanges.Add(FloatStreamId, FrequencyChange);
+			
+			ExpectedStream.FrequencySettings.ObjectOverrides.Remove(MainTestObject);
+			ChangeStreamForSenderClientAndValidate(TEXT("Remove"), ChangeRequest, { ExpectedStream });
+		}
+		
+		return true;
+	}
+
+	/**
+	 * Tests that requests for unregistered objects rejected:
+	 * - Cannot change frequency for unknown stream
+	 * - Cannot change frequency for object not registerd to stream
+	 */
+	IMPLEMENT_CUSTOM_SIMPLE_AUTOMATION_TEST(FRemovingObjectRemovesFrequencyEntry, FChangeStreamsTestBase, "Editor.Concert.Replication.Stream.Frequency.RemovingObjectRemovesFrequencyEntry", EAutomationTestFlags::EditorContext | EAutomationTestFlags::EngineFilter);
+	bool FRemovingObjectRemovesFrequencyEntry::RunTest(const FString& Parameters)
+	{
+		// 1. Start up client with float property stream
+		UTestReflectionObject* MainTestObject = NewObject<UTestReflectionObject>(GetTransientPackage());
+		UTestReflectionObject* SecondaryTestObject = NewObject<UTestReflectionObject>(GetTransientPackage());
+		auto[FloatStreamId, FloatStream] = CreateFloatPropertyStream(*MainTestObject);
+		SenderArgs.Streams = { FloatStream };
+		
+		SetUpClientAndServer();
+		FSharedReplicationStreamDescription ExpectedStream = FloatStream.BaseDescription;
+		
+		// Register SecondaryObject with stream
+		{
+			FConcertReplication_ChangeStream_Request ChangeRequest;
+			
+			FConcertReplication_ChangeStream_PutObject PutObject;
+			AddFloatProperty(PutObject.Properties);
+			PutObject.ClassPath = SecondaryTestObject->GetClass();
+			ChangeRequest.ObjectsToPut.Add({ FloatStreamId, SecondaryTestObject }, PutObject);
+
+			ExpectedStream.ReplicationMap.ReplicatedObjects.Add(SecondaryTestObject, FReplicatedObjectInfo{ PutObject.ClassPath, PutObject.Properties });
+			ChangeStreamForSenderClientAndValidate(TEXT("Add MainTestObject & SecondaryObject"), ChangeRequest, { ExpectedStream });
+		}
+		// Put frequency for MainTestObject & SecondaryObject
+		{
+			FConcertReplication_ChangeStream_Request ChangeRequest;
+			FConcertReplication_ChangeStream_Frequency FrequencyChange;
+			FrequencyChange.OverridesToPut.Add(MainTestObject, { EConcertObjectReplicationMode::SpecifiedRate, 21 });
+			FrequencyChange.OverridesToPut.Add(SecondaryTestObject, { EConcertObjectReplicationMode::SpecifiedRate, 22 });
+			ChangeRequest.FrequencyChanges.Add(FloatStreamId, FrequencyChange);
+			
+			ExpectedStream.FrequencySettings.ObjectOverrides.Add(MainTestObject, { EConcertObjectReplicationMode::SpecifiedRate, 21 });
+			ExpectedStream.FrequencySettings.ObjectOverrides.Add(SecondaryTestObject, { EConcertObjectReplicationMode::SpecifiedRate, 22 });
+			ChangeStreamForSenderClientAndValidate(TEXT("Put SpecifiedRate 21 & 22"), ChangeRequest, { ExpectedStream });
+		}
+		// Removing MainTestObject removes replication frequency entry
+		{
+			FConcertReplication_ChangeStream_Request ChangeRequest;
+			ChangeRequest.ObjectsToRemove.Add({ FloatStreamId, MainTestObject });
+
+			ExpectedStream.ReplicationMap.ReplicatedObjects.Remove(MainTestObject);
+			ExpectedStream.FrequencySettings.ObjectOverrides.Remove(MainTestObject);
+			ChangeStreamForSenderClientAndValidate(TEXT("Unregister object > removes frequency entry"), ChangeRequest, { ExpectedStream });
+		}
+
+		return true;
+	}
+
+	/**
+	 * Tests that requests for unregistered objects rejected:
+	 * - Cannot change frequency for unknown stream
+	 * - Cannot change frequency for object not registerd to stream
+	 */
+	IMPLEMENT_CUSTOM_SIMPLE_AUTOMATION_TEST(FPutReplacesAll, FChangeStreamsTestBase, "Editor.Concert.Replication.Stream.Frequency.PutReplacesAll", EAutomationTestFlags::EditorContext | EAutomationTestFlags::EngineFilter);
+	bool FPutReplacesAll::RunTest(const FString& Parameters)
+	{
+		// 1. Start up client with float property stream
+		UTestReflectionObject* MainTestObject = NewObject<UTestReflectionObject>(GetTransientPackage());
+		UTestReflectionObject* SecondaryTestObject = NewObject<UTestReflectionObject>(GetTransientPackage());
+		auto[FloatStreamId, FloatStream] = CreateFloatPropertyStream(*MainTestObject);
+		SenderArgs.Streams = { FloatStream };
+		SetUpClientAndServer();
+		
+		FSharedReplicationStreamDescription ExpectedStream = FloatStream.BaseDescription;
+		// Register SecondaryObject with stream
+		{
+			FConcertReplication_ChangeStream_Request ChangeRequest;
+			
+			FConcertReplication_ChangeStream_PutObject PutObject;
+			AddFloatProperty(PutObject.Properties);
+			PutObject.ClassPath = SecondaryTestObject->GetClass();
+			ChangeRequest.ObjectsToPut.Add({ FloatStreamId, SecondaryTestObject }, PutObject);
+
+			ExpectedStream.ReplicationMap.ReplicatedObjects.Add(SecondaryTestObject, FReplicatedObjectInfo{ PutObject.ClassPath, PutObject.Properties });
+			ChangeStreamForSenderClientAndValidate(TEXT("Add MainTestObject & SecondaryObject"), ChangeRequest, { ExpectedStream });
+		}
+		
+
+		// Put MainTestObject & SecondaryObject
+		{
+			FConcertReplication_ChangeStream_Request ChangeRequest;
+			FConcertReplication_ChangeStream_Frequency FrequencyChange;
+			FrequencyChange.OverridesToPut.Add(MainTestObject, { EConcertObjectReplicationMode::SpecifiedRate, 10 });
+			FrequencyChange.OverridesToPut.Add(SecondaryTestObject, { EConcertObjectReplicationMode::SpecifiedRate, 20 });
+			ChangeRequest.FrequencyChanges.Add(FloatStreamId, FrequencyChange);
+			
+			ExpectedStream.FrequencySettings.ObjectOverrides.Add(MainTestObject, { EConcertObjectReplicationMode::SpecifiedRate, 10 });
+			ExpectedStream.FrequencySettings.ObjectOverrides.Add(SecondaryTestObject, { EConcertObjectReplicationMode::SpecifiedRate, 20 });
+			ChangeStreamForSenderClientAndValidate(TEXT("Add MainTestObject & SecondaryObject"), ChangeRequest, { ExpectedStream });
+		}
+		
+		// Put replaces all
+		{
+			FConcertReplication_ChangeStream_Request ChangeRequest;
+			FConcertReplication_ChangeStream_Frequency FrequencyChange;
+			FrequencyChange.OverridesToPut.Add(MainTestObject, { EConcertObjectReplicationMode::SpecifiedRate, 100 });
+			ChangeRequest.FrequencyChanges.Add(FloatStreamId, FrequencyChange);
+
+			ExpectedStream.FrequencySettings.ObjectOverrides.Add(MainTestObject, { EConcertObjectReplicationMode::SpecifiedRate, 100 });
+			ExpectedStream.FrequencySettings.ObjectOverrides.Remove(SecondaryTestObject);
+			ChangeStreamForSenderClientAndValidate(TEXT("Put replaces all overrides"), ChangeRequest, { ExpectedStream });
+		}
+		
+		return true;
+	}
+
+	/**
+	 * Tests that requests for unregistered objects rejected:
+	 * - Cannot change frequency for unknown stream
+	 * - Cannot change frequency for object not registerd to stream
+	 */
+	IMPLEMENT_CUSTOM_SIMPLE_AUTOMATION_TEST(FAddAfterPut, FChangeStreamsTestBase, "Editor.Concert.Replication.Stream.Frequency.AddAfterPut", EAutomationTestFlags::EditorContext | EAutomationTestFlags::EngineFilter);
+	bool FAddAfterPut::RunTest(const FString& Parameters)
+	{
+		// 1. Start up client with float property stream
+		UTestReflectionObject* MainTestObject = NewObject<UTestReflectionObject>(GetTransientPackage());
+		auto[FloatStreamId, FloatStream] = CreateFloatPropertyStream(*MainTestObject);
+		SenderArgs.Streams = { FloatStream };
+		SetUpClientAndServer();
+
+		FSharedReplicationStreamDescription ExpectedStream = FloatStream.BaseDescription;
+		// Put SpecifiedRate 25, then add 50
+		{
+			FConcertReplication_ChangeStream_Request ChangeRequest;
+			FConcertReplication_ChangeStream_Frequency FrequencyChange;
+			FrequencyChange.OverridesToPut.Add(MainTestObject, { EConcertObjectReplicationMode::SpecifiedRate, 25 });
+			FrequencyChange.OverridesToAdd.Add(MainTestObject, { EConcertObjectReplicationMode::SpecifiedRate, 50 });
+			ChangeRequest.FrequencyChanges.Add(FloatStreamId, FrequencyChange);
+			
+			ExpectedStream.FrequencySettings.ObjectOverrides.Add(MainTestObject, { EConcertObjectReplicationMode::SpecifiedRate, 50 });
+			ChangeStreamForSenderClientAndValidate(TEXT("Put SpecifiedRate 25, then add 50"), ChangeRequest, { ExpectedStream });
+		}
+		
+		return true;
+	}
+
+	/**
+	 * Tests that requests for unregistered objects rejected:
+	 * - Cannot change frequency for unknown stream
+	 * - Cannot change frequency for object not registerd to stream
+	 */
+	IMPLEMENT_CUSTOM_SIMPLE_AUTOMATION_TEST(FCannotChangeFrequencyOfUnregistered, FChangeStreamsTestBase, "Editor.Concert.Replication.Stream.Frequency.FailWithUnregisteredStreamOrObject", EAutomationTestFlags::EditorContext | EAutomationTestFlags::EngineFilter);
+	bool FCannotChangeFrequencyOfUnregistered::RunTest(const FString& Parameters)
+	{
+		// 1. Start up client with float property stream
+		UTestReflectionObject* MainTestObject = NewObject<UTestReflectionObject>(GetTransientPackage());
+		UTestReflectionObject* SecondaryTestObject = NewObject<UTestReflectionObject>(GetTransientPackage());
+		auto[FloatStreamId, FloatStream] = CreateFloatPropertyStream(*MainTestObject);
+		SenderArgs.Streams = { FloatStream };
+		
+		SetUpClientAndServer();
+		// We'll intentionally be causing errors below that cause a warning to be logged.
+		AddExpectedError(TEXT("Rejecting ChangeStream"),EAutomationExpectedErrorFlags::Contains, 2);
+
+		// Change defaults to be SpecifiedRate at 42 update rate.
+		FSharedReplicationStreamDescription ExpectedStream = FloatStream.BaseDescription;
+		
+		// Cannot change frequency of unregistered stream
+		{
+			const FGuid UnregisteredStreamId = FGuid::NewGuid();
+			
+			FConcertReplication_ChangeStream_Request ChangeRequest;
+			FConcertReplication_ChangeStream_Frequency FrequencyChange;
+			FrequencyChange.OverridesToPut.Add(MainTestObject, { EConcertObjectReplicationMode::SpecifiedRate, 50 });
+			ChangeRequest.FrequencyChanges.Add(UnregisteredStreamId, FrequencyChange);
+
+			ChangeStreamForSenderClientAndValidate(
+				TEXT("Cannot change frequency of unregistered stream"),
+				ChangeRequest,
+				{ ExpectedStream },
+				EChangeStreamValidation::TestWasHandled | EChangeStreamValidation::TestResponseFailed
+				)
+				.Next([this, &UnregisteredStreamId, &MainTestObject](FConcertReplication_ChangeStream_Response&& Response)
+				{
+					const EConcertChangeObjectFrequencyErrorCode* ErrorCode = Response.FrequencyErrors.OverrideFailures.Find({ UnregisteredStreamId, MainTestObject });
+					TestTrue(TEXT("Cannot change frequency of unregistered stream > Rejection code (defaults)"), Response.FrequencyErrors.DefaultFailures.IsEmpty());
+					TestTrue(TEXT("Cannot change frequency of unregistered stream > Rejection code (object)"), ErrorCode && *ErrorCode == EConcertChangeObjectFrequencyErrorCode::NotRegistered);
+				});
+		}
+
+		// Cannot add object that is not registered
+		{
+			FConcertReplication_ChangeStream_Request ChangeRequest;
+			FConcertReplication_ChangeStream_Frequency FrequencyChange;
+			FrequencyChange.OverridesToPut.Add(SecondaryTestObject, { EConcertObjectReplicationMode::SpecifiedRate, 100 });
+			ChangeRequest.FrequencyChanges.Add(FloatStreamId, FrequencyChange);
+			
+			ChangeStreamForSenderClientAndValidate(
+				TEXT("Cannot add non-registered object"),
+				ChangeRequest,
+				{ ExpectedStream },
+				EChangeStreamValidation::TestWasHandled | EChangeStreamValidation::TestResponseFailed
+				)
+				.Next([this, FloatStreamId, SecondaryTestObject](FConcertReplication_ChangeStream_Response&& Response)
+				{
+					const EConcertChangeObjectFrequencyErrorCode* ErrorCode = Response.FrequencyErrors.OverrideFailures.Find({ FloatStreamId, SecondaryTestObject });
+					TestTrue(TEXT("Cannot add non-registered object > Rejection code (defaults)"), Response.FrequencyErrors.DefaultFailures.IsEmpty());
+					TestTrue(TEXT("Cannot add non-registered object > Rejection code (object)"), ErrorCode && *ErrorCode == EConcertChangeObjectFrequencyErrorCode::NotRegistered);
+				});
+		}
+
+		return true;
+	}
+}
\ No newline at end of file
diff --git a/Engine/Plugins/Developer/Concert/ConcertSync/ConcertSyncTest/Source/ConcertSyncTest/Private/Replication/StreamRequestTests_Query.cpp b/Engine/Plugins/Developer/Concert/ConcertSync/ConcertSyncTest/Source/ConcertSyncTest/Private/Replication/StreamRequestTests_Query.cpp
new file mode 100644
--- /dev/null
+++ b/Engine/Plugins/Developer/Concert/ConcertSync/ConcertSyncTest/Source/ConcertSyncTest/Private/Replication/StreamRequestTests_Query.cpp
@@ -0,0 +1,198 @@
+// Copyright Epic Games, Inc. All Rights Reserved.
+
+#include "Util/ClientServerCommunicationTest.h"
+
+#include "Replication/Data/ReplicationStreamDescription.h"
+#include "Replication/IConcertClientReplicationManager.h"
+#include "TestReflectionObject.h"
+#include "Util/ChangeStreamsTestBase.h"
+#include "Util/SendReceiveObjectTestBase.h"
+
+#include "Misc/AutomationTest.h"
+#include "UObject/UObjectGlobals.h"
+#include "UObject/Package.h"
+
+namespace UE::ConcertSyncTests::Replication::Stream
+{
+	/**
+	 * Let's receiving client query the server what the sending client is sending. Checks
+	 * - the reported streams,
+	 * - that the authority updates,
+	 * - the EConcertQueryClientStreamFlags flags
+	 *
+	 * @note The other tests in StreamRequestTests_{x}.cpp use QueryStream under the hood so all this test does is checking that the skip flags work correctly. 
+	 */
+	IMPLEMENT_CUSTOM_SIMPLE_AUTOMATION_TEST(FQueryOtherClientStreams, FSendReceiveObjectTestBase, "Editor.Concert.Replication.Stream.QueryClientStream", EAutomationTestFlags::EditorContext | EAutomationTestFlags::EngineFilter);
+	bool FQueryOtherClientStreams::RunTest(const FString& Parameters)
+	{
+		// 1. Init
+		SetUpClientAndServer();
+		
+		const FGuid SenderEndpointId = Client_Sender->ClientSessionMock->GetSessionClientEndpointId();
+		const FSoftObjectPath TestObjectPath{ TestObject };
+		FConcertReplication_QueryReplicationInfo_Request Request;
+		Request.ClientEndpointIds = { SenderEndpointId };
+		auto TestReplicationMapContent = [this, &SenderEndpointId](const FConcertReplication_QueryReplicationInfo_Response& Response)
+		{
+			TestTrue(TEXT("ErrorCode == Handled"), Response.ErrorCode == EReplicationResponseErrorCode::Handled);
+			
+			const FReplicationClientQueriedInfo* Info = Response.ClientInfo.Find(SenderEndpointId);
+			if (!Info)
+			{
+				AddError(TEXT("No info about sending client!"));
+				return;
+			}
+				
+			TestEqual(TEXT("Only received info about requested client"), Response.ClientInfo.Num(), 1);
+			TestEqual(TEXT("Exactly one stream data"), Info->Streams.Num(), 1);
+			if (Info->Streams.IsEmpty())
+			{
+				return;
+			}
+				
+			const FSharedReplicationStreamDescription& StreamDescription = Info->Streams[0];
+			TestEqual(TEXT("StreamDescription->Identifier correct"), StreamDescription.Identifier, SenderStreamId);
+			TestEqual(TEXT("StreamDescription->ReplicationMap has exactly 1 object"), StreamDescription.ReplicationMap.ReplicatedObjects.Num(), 1);
+
+			const FObjectReplicationMap ExpectedReplicationMap = CreateSenderArgs().Streams[0].BaseDescription.ReplicationMap;
+			TestEqual(TEXT("Registered and reported replication maps match"), StreamDescription.ReplicationMap, ExpectedReplicationMap);
+		};
+		
+		// 2.1 Request before taking authority
+		bool bReceivedFirstQueryResponse = false;
+		ClientReplicationManager_Receiver->QueryClientInfo(Request)
+			.Next([this, &SenderEndpointId, &TestReplicationMapContent, &bReceivedFirstQueryResponse](FConcertReplication_QueryReplicationInfo_Response&& Response) mutable
+			{
+				bReceivedFirstQueryResponse = true;
+				TestTrue(TEXT("ErrorCode == Handled"), Response.ErrorCode == EReplicationResponseErrorCode::Handled);
+				
+				TestReplicationMapContent(Response);
+				if (const FReplicationClientQueriedInfo* Info = Response.ClientInfo.Find(SenderEndpointId))
+				{
+					TestEqual(TEXT("Contains no authority data"), Info->Authority.Num(), 0);
+				}
+			});
+		TestTrue(TEXT("Received query response (before taking authority)"), bReceivedFirstQueryResponse);
+
+		// 2.2 Request after taking authority
+		bool bSenderReceivedResponse = false;
+		ClientReplicationManager_Sender->TakeAuthorityOver({ TestObjectPath })
+			.Next([this, &bSenderReceivedResponse](const FConcertReplication_ChangeAuthority_Response& Response) mutable
+			{
+				bSenderReceivedResponse = true;
+				TestEqual(TEXT("No rejection taking authority"), Response.RejectedObjects.Num(), 0); 
+				TestTrue(TEXT("ErrorCode == Handled"), Response.ErrorCode == EReplicationResponseErrorCode::Handled);
+			});
+		TestTrue(TEXT("Received response taking authority"), bSenderReceivedResponse);
+
+		bool bReceivedSecondQueryResponse = false;
+		ClientReplicationManager_Receiver->QueryClientInfo(Request)
+			.Next([this, &SenderEndpointId, &TestReplicationMapContent, &TestObjectPath, &bReceivedSecondQueryResponse](FConcertReplication_QueryReplicationInfo_Response&& Response) mutable
+			{
+				bReceivedSecondQueryResponse = true;
+				TestTrue(TEXT("ErrorCode == Handled"), Response.ErrorCode == EReplicationResponseErrorCode::Handled);
+				TestReplicationMapContent(Response);
+				
+				const FReplicationClientQueriedInfo* Info = Response.ClientInfo.Find(SenderEndpointId);
+				if (!Info || Info->Authority.IsEmpty())
+				{
+					AddError(TEXT("Missing authority data"));
+					return;
+				}
+
+				TestEqual(TEXT("Exactly 1 authority stream"), Info->Authority.Num(), 1);
+				const FReplicationAuthorityInfo& AuthorityInfo = Info->Authority[0];
+				TestEqual(TEXT("Authority stream ID matches registered stream ID"), AuthorityInfo.StreamId, SenderStreamId);
+				TestEqual(TEXT("Has authority over exactly 1 object"), AuthorityInfo.AuthoredObjects.Num(), 1);
+				TestTrue(TEXT("Has authority over registered object"), AuthorityInfo.AuthoredObjects.Contains(TestObjectPath));
+			});
+		TestTrue(TEXT("Received response to taking authority"), bSenderReceivedResponse);
+		TestTrue(TEXT("Received query response (after taking authority)"), bReceivedSecondQueryResponse);
+
+		// 2.3 SkipStreamInfo
+		bool bReceivedResponse_SkipStreamInfo = false;
+		Request.QueryFlags = EConcertQueryClientStreamFlags::SkipStreamInfo;
+		ClientReplicationManager_Receiver->QueryClientInfo(Request)
+			.Next([this, &SenderEndpointId, &bReceivedResponse_SkipStreamInfo](FConcertReplication_QueryReplicationInfo_Response&& Response) mutable
+			{
+				bReceivedResponse_SkipStreamInfo = true;
+				TestTrue(TEXT("ErrorCode == Handled"), Response.ErrorCode == EReplicationResponseErrorCode::Handled);
+				if (const FReplicationClientQueriedInfo* Info = Response.ClientInfo.Find(SenderEndpointId))
+				{
+					TestEqual(TEXT("SkipStreamInfo > No stream data"), Info->Streams.Num(), 0);
+				}
+			});
+		TestTrue(TEXT("bReceivedResponse_SkipStreamInfo"), bReceivedResponse_SkipStreamInfo);
+		
+		// 2.4 SkipProperties
+		bool bReceivedResponse_SkipProperties = false;
+		Request.QueryFlags = EConcertQueryClientStreamFlags::SkipProperties;
+		ClientReplicationManager_Receiver->QueryClientInfo(Request)
+			.Next([this, &SenderEndpointId, &TestObjectPath, &bReceivedResponse_SkipProperties](FConcertReplication_QueryReplicationInfo_Response&& Response) mutable
+			{
+				bReceivedResponse_SkipProperties = true;
+				const FReplicationClientQueriedInfo* Info = Response.ClientInfo.Find(SenderEndpointId);
+				TestTrue(TEXT("ErrorCode == Handled"), Response.ErrorCode == EReplicationResponseErrorCode::Handled);
+				if (!Info || Info->Streams.IsEmpty())
+				{
+					AddError(TEXT("SkipProperties > No stream info"));
+					return;
+				}
+
+				const FReplicatedObjectInfo* ObjectInfo = Info->Streams[0].ReplicationMap.ReplicatedObjects.Find(TestObjectPath);
+				if (!ObjectInfo)
+				{
+					AddError(TEXT("SkipProperties > No object info"));
+					return;
+				}
+				TestEqual(TEXT("SkipProperties > No property data"), ObjectInfo->PropertySelection.ReplicatedProperties.Num(), 0);
+			});
+		TestTrue(TEXT("bReceivedResponse_SkipProperties"), bReceivedResponse_SkipProperties);
+		
+		// 2.5 SkipAuthority
+		bool bReceivedResponse_SkipAuthority = false;
+		Request.QueryFlags = EConcertQueryClientStreamFlags::SkipAuthority;
+		ClientReplicationManager_Receiver->QueryClientInfo(Request)
+			.Next([this, &SenderEndpointId, &bReceivedResponse_SkipAuthority](FConcertReplication_QueryReplicationInfo_Response&& Response) mutable
+			{
+				bReceivedResponse_SkipAuthority = true;
+				TestTrue(TEXT("ErrorCode == Handled"), Response.ErrorCode == EReplicationResponseErrorCode::Handled);
+				if (const FReplicationClientQueriedInfo* Info = Response.ClientInfo.Find(SenderEndpointId))
+				{
+					TestEqual(TEXT("SkipAuthority > No authority data"), Info->Authority.Num(), 0);
+				}
+			});
+		TestTrue(TEXT("bReceivedResponse_SkipAuthority"), bReceivedResponse_SkipAuthority);
+
+		// 2.6 SkipFrequency
+		FConcertReplication_ChangeStream_Request AddFrequencyInfo;
+		FConcertReplication_ChangeStream_Frequency& FrequencyChange = AddFrequencyInfo.FrequencyChanges.Add(SenderStreamId);
+		FrequencyChange.OverridesToPut.Add(TestObjectPath) = { EConcertObjectReplicationMode::SpecifiedRate, 20 }; 
+		FrequencyChange.NewDefaults = { EConcertObjectReplicationMode::SpecifiedRate, 42 };
+		FrequencyChange.Flags = EConcertReplicationChangeFrequencyFlags::SetDefaults;
+		// Important: Change Sender here, not Receiver! Receiver is the one querying for Sender.
+		ClientReplicationManager_Sender->ChangeStream(AddFrequencyInfo)
+			.Next([this](FConcertReplication_ChangeStream_Response&& Response)
+			{
+				TestTrue(TEXT("Changed frequency"), Response.IsSuccess());
+			});
+		
+		bool bReceivedResponse_SkipFrequency = false;
+		Request.QueryFlags = EConcertQueryClientStreamFlags::SkipFrequency;
+		ClientReplicationManager_Receiver->QueryClientInfo(Request)
+			.Next([this, &SenderEndpointId, &bReceivedResponse_SkipFrequency](FConcertReplication_QueryReplicationInfo_Response&& Response) mutable
+			{
+				bReceivedResponse_SkipFrequency = true;
+				TestTrue(TEXT("ErrorCode == Handled"), Response.ErrorCode == EReplicationResponseErrorCode::Handled);
+				if (const FReplicationClientQueriedInfo* Info = Response.ClientInfo.Find(SenderEndpointId)
+					; Info && ensureAlways(Info->Streams.Num() == 1))
+				{
+					const FConcertStreamFrequencySettings& FrequencySettings = Info->Streams[0].FrequencySettings;
+					TestEqual(TEXT("SkipFrequency > No frequency override data"), FrequencySettings.ObjectOverrides.Num(), 0);
+				}
+			});
+		TestTrue(TEXT("bReceivedResponse_SkipFrequency"), bReceivedResponse_SkipFrequency);
+		
+		return true;
+	}
+}
\ No newline at end of file
diff --git a/Engine/Plugins/Developer/Concert/ConcertSync/ConcertSyncTest/Source/ConcertSyncTest/Private/Replication/Util/ChangeStreamsTestBase.cpp b/Engine/Plugins/Developer/Concert/ConcertSync/ConcertSyncTest/Source/ConcertSyncTest/Private/Replication/Util/ChangeStreamsTestBase.cpp
new file mode 100644
--- /dev/null
+++ b/Engine/Plugins/Developer/Concert/ConcertSync/ConcertSyncTest/Source/ConcertSyncTest/Private/Replication/Util/ChangeStreamsTestBase.cpp
@@ -0,0 +1,138 @@
+// Copyright Epic Games, Inc. All Rights Reserved.
+
+#include "ChangeStreamsTestBase.h"
+
+#include "Replication/Data/ReplicationStreamDescription.h"
+#include "Replication/IConcertClientReplicationManager.h"
+#include "Replication/TestReflectionObject.h"
+#include "Replication/Util/SendReceiveGenericStreamTestBase.h"
+
+#include "Misc/AutomationTest.h"
+#include "UObject/UObjectGlobals.h"
+#include "UObject/Package.h"
+
+namespace UE::ConcertSyncTests::Replication
+{
+	FChangeStreamsTestBase::FChangeStreamsTestBase(const FString& InName, const bool bInComplexTask)
+			: FSendReceiveGenericTestBase(InName, bInComplexTask)
+	{}
+	
+	void FChangeStreamsTestBase::AddFloatProperties(const UTestReflectionObject& TestObject, FReplicationStreamDescription& Stream)
+	{
+		FConcertPropertySelection Properties;
+		AddFloatProperty(Properties);
+		const FReplicatedObjectInfo AllProperties { TestObject.GetClass(), MoveTemp(Properties) };
+		Stream.BaseDescription.ReplicationMap.ReplicatedObjects.Add(&TestObject, AllProperties);
+	}
+	
+	TTuple<FGuid, FReplicationStreamDescription> FChangeStreamsTestBase::CreateFloatPropertyStream(
+		const UTestReflectionObject& TestObject,
+		FConcertObjectReplicationSettings DefaultFrequencySettings
+		)
+	{
+		const FGuid NewStreamId = FGuid::NewGuid();
+
+		// Use Realtime replication because otherwise our test will need to pass latent time due to the frequency system
+		FReplicationStreamDescription SendingStream;
+		SendingStream.BaseDescription.Identifier = NewStreamId;
+		SendingStream.BaseDescription.FrequencySettings.Defaults = DefaultFrequencySettings;
+		AddFloatProperties(TestObject, SendingStream);
+		return { NewStreamId, SendingStream };
+	}
+
+	TTuple<FGuid, FReplicationStreamDescription> FChangeStreamsTestBase::CreateVectorPropertyStream(
+		const UTestReflectionObject& TestObject,
+		FConcertObjectReplicationSettings DefaultFrequencySettings
+		)
+	{
+		const FGuid NewStreamId = FGuid::NewGuid();
+		FConcertPropertySelection Properties;
+		AddVectorProperty(Properties);
+		
+		// Use Realtime replication because otherwise our test will need to pass latent time due to the frequency system
+		const FReplicatedObjectInfo AllProperties { TestObject.GetClass(), MoveTemp(Properties) };
+		FReplicationStreamDescription SendingStream;
+		SendingStream.BaseDescription.Identifier = NewStreamId;
+		SendingStream.BaseDescription.ReplicationMap.ReplicatedObjects.Add(&TestObject, AllProperties);
+		SendingStream.BaseDescription.FrequencySettings.Defaults = DefaultFrequencySettings;
+		return { NewStreamId, SendingStream };
+	}
+
+	FConcertPropertySelection& FChangeStreamsTestBase::GetPropertySelection(FReplicationStreamDescription& Stream, const UTestReflectionObject& TestObject)
+	{
+		return Stream.BaseDescription.ReplicationMap.ReplicatedObjects.FindOrAdd(&TestObject).PropertySelection;
+	}
+
+	void FChangeStreamsTestBase::AddFloatProperty(FConcertPropertySelection& PropertySelection)
+	{
+		const FProperty* FloatProperty = UTestReflectionObject::StaticClass()->FindPropertyByName(GET_MEMBER_NAME_CHECKED(UTestReflectionObject, Float));
+		const FConcertPropertyChain PropertyChain(nullptr, *FloatProperty);
+		PropertySelection.ReplicatedProperties.Add(PropertyChain);
+	}
+
+	void FChangeStreamsTestBase::AddVectorProperty(FConcertPropertySelection& PropertySelection)
+	{
+		PropertySelection.ReplicatedProperties.Add(*FConcertPropertyChain::CreateFromPath(*UTestReflectionObject::StaticClass(), { TEXT("Vector") }));
+		PropertySelection.ReplicatedProperties.Add(*FConcertPropertyChain::CreateFromPath(*UTestReflectionObject::StaticClass(), { TEXT("Vector"), TEXT("X") }));
+	}
+
+	/** Sends a ChangeStream request and validates that ExpectedStreams is equal to the registered streams. */
+	TFuture<FConcertReplication_ChangeStream_Response> FChangeStreamsTestBase::ChangeStreamForSenderClientAndValidate(
+		const FString& InTestName,
+		const FConcertReplication_ChangeStream_Request& Request,
+		const TArray<FSharedReplicationStreamDescription>& ExpectedStreams,
+		EChangeStreamValidation ValidationFlags
+		)
+	{
+		using namespace ConcertSyncClient::Replication;
+		
+		bool bModifiedInitialStream = false;
+		TFuture<FConcertReplication_ChangeStream_Response> Future = ClientReplicationManager_Sender->ChangeStream(Request)
+			.Next([this, &InTestName, ValidationFlags, &bModifiedInitialStream](FConcertReplication_ChangeStream_Response&& Response)
+			{
+				bModifiedInitialStream = true;
+				
+				const bool bValidatedErrorCode =
+					!EnumHasAnyFlags(ValidationFlags, EChangeStreamValidation::TestWasHandled | EChangeStreamValidation::TestWasTimeout)
+					|| (EnumHasAnyFlags(ValidationFlags, EChangeStreamValidation::TestWasHandled) && Response.ErrorCode == EReplicationResponseErrorCode::Handled)
+					|| (EnumHasAnyFlags(ValidationFlags, EChangeStreamValidation::TestWasTimeout) && Response.ErrorCode == EReplicationResponseErrorCode::Timeout);
+				TestTrue(FString::Printf(TEXT("%s: Validate error code"), *InTestName), bValidatedErrorCode);
+				
+				const bool bResponseValidationResult =
+					!EnumHasAnyFlags(ValidationFlags, EChangeStreamValidation::TestResponseSuccess | EChangeStreamValidation::TestResponseFailed)
+					|| (EnumHasAnyFlags(ValidationFlags, EChangeStreamValidation::TestResponseSuccess) && Response.IsSuccess())
+					|| (EnumHasAnyFlags(ValidationFlags, EChangeStreamValidation::TestResponseFailed) && Response.IsFailure());
+				TestTrue(FString::Printf(TEXT("%s: Validate response"), *InTestName), bResponseValidationResult);
+				
+				return Response;
+			});
+		TestTrue(FString::Printf(TEXT("%s > Modify Stream > Received response"), *InTestName), bModifiedInitialStream);
+		ValidateSenderClientStreams(InTestName, ExpectedStreams);
+		
+		return MoveTemp(Future);
+	}
+	
+	void FChangeStreamsTestBase::ValidateSenderClientStreams(const FString& InTestName, const TArray<FSharedReplicationStreamDescription>& ExpectedStreams)
+	{
+		using namespace ConcertSyncClient::Replication;
+
+		// Test that remote clients see the receive the same as Streams
+		bool bQueriedClientInfo = false;
+		const FGuid SenderId = Client_Sender->ClientSessionMock->GetSessionClientEndpointId();
+		ClientReplicationManager_Receiver->QueryClientInfo({{ { SenderId } }})
+			.Next([this, &InTestName, &bQueriedClientInfo, &ExpectedStreams, &SenderId](FConcertReplication_QueryReplicationInfo_Response&& Response)
+			{
+				bQueriedClientInfo = true;
+				TestTrue(TEXT("ErrorCode == Handled"), Response.ErrorCode == EReplicationResponseErrorCode::Handled);
+				const FReplicationClientQueriedInfo* Info = Response.ClientInfo.Find(SenderId);
+				const bool bRegisteredAndQueriedStreamsAreEqual = Info->Streams == ExpectedStreams;
+				TestTrue(FString::Printf(TEXT("%s > Registered and queried streams are equal"), *InTestName), bRegisteredAndQueriedStreamsAreEqual);
+			});
+
+		// Validate the local client's cache equals  Streams
+		const TArray<FReplicationStreamDescription> LocalStreams = ClientReplicationManager_Sender->GetRegisteredStreams();
+		TArray<FSharedReplicationStreamDescription> TransformedLocalStreams;
+		Algo::Transform(LocalStreams, TransformedLocalStreams, [](const FReplicationStreamDescription& Stream){ return Stream.BaseDescription; });
+		TestEqual(FString::Printf(TEXT("%s > Local streams match registered streams"), *InTestName), TransformedLocalStreams, ExpectedStreams);
+	};
+}
\ No newline at end of file
diff --git a/Engine/Plugins/Developer/Concert/ConcertSync/ConcertSyncTest/Source/ConcertSyncTest/Private/Replication/Util/ChangeStreamsTestBase.h b/Engine/Plugins/Developer/Concert/ConcertSync/ConcertSyncTest/Source/ConcertSyncTest/Private/Replication/Util/ChangeStreamsTestBase.h
new file mode 100644
--- /dev/null
+++ b/Engine/Plugins/Developer/Concert/ConcertSync/ConcertSyncTest/Source/ConcertSyncTest/Private/Replication/Util/ChangeStreamsTestBase.h
@@ -0,0 +1,80 @@
+// Copyright Epic Games, Inc. All Rights Reserved.
+
+#pragma once
+
+#include "Replication/Data/ReplicationStreamDescription.h"
+#include "Replication/IConcertClientReplicationManager.h"
+#include "Replication/Messages/ChangeStream.h"
+#include "SendReceiveGenericStreamTestBase.h"
+
+#include "Async/Future.h"
+
+class UTestReflectionObject;
+
+namespace UE::ConcertSyncTests::Replication
+{
+	enum class EChangeStreamValidation : uint8
+	{
+		/** Do special flags */
+		None = 0,
+
+		/** Test that the ErrorCode is EReplicationResponseErrorCode::Handled. */
+		TestWasHandled = 1 << 0,
+		/** Test that the ErrorCode is EReplicationResponseErrorCode::Timeout. */
+		TestWasTimeout = 1 << 1,
+		
+		/** Test that the stream change request succeeded */
+		TestResponseSuccess = 1 << 2,
+		/** Test that the stream change request  failed */
+		TestResponseFailed = 1 << 3,
+	};
+	ENUM_CLASS_FLAGS(EChangeStreamValidation);
+	
+	/**
+	 * Base test for testing ChangeStream requests.
+	 * 
+	 * While there is no sending nor receiving of replicated data, subclasses will still refer to the clients as "Sender"
+	 * and "Receiver".
+	 */
+	class FChangeStreamsTestBase : public FSendReceiveGenericTestBase
+	{
+	public:
+
+		FChangeStreamsTestBase(const FString& InName, const bool bInComplexTask);
+
+	protected:
+		
+		ConcertSyncClient::Replication::FJoinReplicatedSessionArgs SenderArgs;
+		ConcertSyncClient::Replication::FJoinReplicatedSessionArgs ReceiverArgs;
+
+		//~ Begin FSendReceiveTestBase Interface
+		virtual ConcertSyncClient::Replication::FJoinReplicatedSessionArgs CreateSenderArgs() override { return SenderArgs; }
+		virtual ConcertSyncClient::Replication::FJoinReplicatedSessionArgs CreateReceiverArgs() override { return ReceiverArgs; }
+		//~ End FSendReceiveTestBase Interface
+
+		static void AddFloatProperties(const UTestReflectionObject& TestObject, FReplicationStreamDescription& Stream);
+		
+		static TTuple<FGuid, FReplicationStreamDescription> CreateFloatPropertyStream(
+			const UTestReflectionObject& TestObject,
+			FConcertObjectReplicationSettings DefaultFrequencySettings = { EConcertObjectReplicationMode::Realtime }
+			);
+		static TTuple<FGuid, FReplicationStreamDescription> CreateVectorPropertyStream(
+			const UTestReflectionObject& TestObject,
+			FConcertObjectReplicationSettings DefaultFrequencySettings = { EConcertObjectReplicationMode::Realtime }
+			);
+		
+		static FConcertPropertySelection& GetPropertySelection(FReplicationStreamDescription& Stream, const UTestReflectionObject& TestObject);
+
+		static void AddFloatProperty(FConcertPropertySelection& PropertySelection);
+		static void AddVectorProperty(FConcertPropertySelection& PropertySelection);
+
+		/** Sends a ChangeStream request and validates that ExpectedStreams is equal to the registered streams. */
+		TFuture<FConcertReplication_ChangeStream_Response> ChangeStreamForSenderClientAndValidate(
+			const FString& InTestName,
+			const FConcertReplication_ChangeStream_Request& Request,
+			const TArray<FSharedReplicationStreamDescription>& ExpectedStreams,
+			EChangeStreamValidation ValidationFlags = EChangeStreamValidation::TestWasHandled | EChangeStreamValidation::TestResponseSuccess
+			);
+		void ValidateSenderClientStreams(const FString& InTestName, const TArray<FSharedReplicationStreamDescription>& ExpectedStreams);
+	};
+}
diff --git a/Engine/Plugins/Developer/Concert/ConcertSync/ConcertSyncTest/Source/ConcertSyncTest/Private/Replication/Util/SendReceiveTestBase.cpp b/Engine/Plugins/Developer/Concert/ConcertSync/ConcertSyncTest/Source/ConcertSyncTest/Private/Replication/Util/SendReceiveTestBase.cpp
--- a/Engine/Plugins/Developer/Concert/ConcertSync/ConcertSyncTest/Source/ConcertSyncTest/Private/Replication/Util/SendReceiveTestBase.cpp
+++ b/Engine/Plugins/Developer/Concert/ConcertSync/ConcertSyncTest/Source/ConcertSyncTest/Private/Replication/Util/SendReceiveTestBase.cpp
@@ -22,20 +22,26 @@ namespace UE::ConcertSyncTests::Replication
 			: FConcertClientServerCommunicationTest(InName, bInComplexTask)
 	{}
 
-	ConcertSyncClient::Replication::FJoinReplicatedSessionArgs FSendReceiveTestBase::CreateHandshakeArgsFrom(const UObject& Object, const FGuid& SenderStreamId)
+	ConcertSyncClient::Replication::FJoinReplicatedSessionArgs FSendReceiveTestBase::CreateHandshakeArgsFrom(
+		const UObject& Object,
+		const FGuid& SenderStreamId,
+		EConcertObjectReplicationMode ReplicationMode,
+		uint8 ReplicationRate
+		)
 	{
 		ConcertSyncClient::Replication::FJoinReplicatedSessionArgs SenderJoinArgs;
 		
-		FReplicatedObjectInfo AllProperties { Object.GetClass() };
-		ConcertSyncCore::PropertyChain::ForEachReplicatableConcertProperty(*Object.GetClass(), [&AllProperties](FConcertPropertyChain&& Chain)
+		FReplicatedObjectInfo ReplicatedObjectInfo { Object.GetClass() };
+		ConcertSyncCore::PropertyChain::ForEachReplicatableConcertProperty(*Object.GetClass(), [&ReplicatedObjectInfo](FConcertPropertyChain&& Chain)
 		{
-			AllProperties.PropertySelection.ReplicatedProperties.Emplace(MoveTemp(Chain));
+			ReplicatedObjectInfo.PropertySelection.ReplicatedProperties.Emplace(MoveTemp(Chain));
 			return EBreakBehavior::Continue;
 		});
 
 		FReplicationStreamDescription SendingStream;
 		SendingStream.BaseDescription.Identifier = SenderStreamId;
-		SendingStream.BaseDescription.ReplicationMap.ReplicatedObjects.Add(&Object, AllProperties);
+		SendingStream.BaseDescription.ReplicationMap.ReplicatedObjects.Add(&Object, ReplicatedObjectInfo);
+		SendingStream.BaseDescription.FrequencySettings.Defaults = { ReplicationMode, ReplicationRate };
 		SenderJoinArgs.Streams.Add(SendingStream);
 		return SenderJoinArgs;
 	}
diff --git a/Engine/Plugins/Developer/Concert/ConcertSync/ConcertSyncTest/Source/ConcertSyncTest/Private/Replication/Util/SendReceiveTestBase.h b/Engine/Plugins/Developer/Concert/ConcertSync/ConcertSyncTest/Source/ConcertSyncTest/Private/Replication/Util/SendReceiveTestBase.h
--- a/Engine/Plugins/Developer/Concert/ConcertSync/ConcertSyncTest/Source/ConcertSyncTest/Private/Replication/Util/SendReceiveTestBase.h
+++ b/Engine/Plugins/Developer/Concert/ConcertSync/ConcertSyncTest/Source/ConcertSyncTest/Private/Replication/Util/SendReceiveTestBase.h
@@ -28,8 +28,20 @@ namespace UE::ConcertSyncTests::Replication
 	public:
 	
 		FSendReceiveTestBase(const FString& InName, const bool bInComplexTask);
-		
-		static ConcertSyncClient::Replication::FJoinReplicatedSessionArgs CreateHandshakeArgsFrom(const UObject& Object, const FGuid& SenderStreamId = FGuid::NewGuid());
+
+		/**
+		 * Creates joining arguments containing a stream that will replicate a single object.
+		 * @param Object The object to replicate
+		 * @param SenderStreamId The ID of the new stream
+		 * @param ReplicationMode The replication mode for the object. Set to Realtime by default so tests need not handle latency (due to FReplicatedObjectInfo::ReplicationRate).
+		 * @param ReplicationRate The replication rate for the object. Only used if ReplicationMode == EConcertObjectReplicationMode::SpecifiedRate
+		 */
+		static ConcertSyncClient::Replication::FJoinReplicatedSessionArgs CreateHandshakeArgsFrom(
+			const UObject& Object,
+			const FGuid& SenderStreamId = FGuid::NewGuid(),
+			EConcertObjectReplicationMode ReplicationMode = EConcertObjectReplicationMode::Realtime,
+			uint8 ReplicationRate = 30
+			);
 		
 		using FReceiveReplicationEventSignature = void(const FConcertSessionContext& Context, const FConcertReplication_BatchReplicationEvent& Event);
 		
