Fix capturing structured binding in lambda

diff --git a/Engine/Plugins/Developer/Concert/ConcertSync/ConcertSyncTest/Source/ConcertSyncTest/Private/Replication/StreamRequestTests_Frequency.cpp b/Engine/Plugins/Developer/Concert/ConcertSync/ConcertSyncTest/Source/ConcertSyncTest/Private/Replication/StreamRequestTests_Frequency.cpp
--- a/Engine/Plugins/Developer/Concert/ConcertSync/ConcertSyncTest/Source/ConcertSyncTest/Private/Replication/StreamRequestTests_Frequency.cpp
+++ b/Engine/Plugins/Developer/Concert/ConcertSync/ConcertSyncTest/Source/ConcertSyncTest/Private/Replication/StreamRequestTests_Frequency.cpp
@@ -28,7 +28,9 @@ namespace UE::ConcertSyncTests::Replication::Stream
 		// 1. Start up client with float property stream
 		UTestReflectionObject* MainTestObject = NewObject<UTestReflectionObject>(GetTransientPackage());
 		UTestReflectionObject* SecondaryTestObject = NewObject<UTestReflectionObject>(GetTransientPackage());
-		auto[FloatStreamId, FloatStream] = CreateFloatPropertyStream(*MainTestObject);
+		FGuid FloatStreamId;
+		FReplicationStreamDescription FloatStream;
+		Tie(FloatStreamId, FloatStream) = CreateFloatPropertyStream(*MainTestObject);
 		SenderArgs.Streams = { FloatStream };
 		
 		SetUpClientAndServer();
@@ -318,7 +320,9 @@ namespace UE::ConcertSyncTests::Replication::Stream
 		// 1. Start up client with float property stream
 		UTestReflectionObject* MainTestObject = NewObject<UTestReflectionObject>(GetTransientPackage());
 		UTestReflectionObject* SecondaryTestObject = NewObject<UTestReflectionObject>(GetTransientPackage());
-		auto[FloatStreamId, FloatStream] = CreateFloatPropertyStream(*MainTestObject);
+		FGuid FloatStreamId;
+		FReplicationStreamDescription FloatStream;
+		Tie(FloatStreamId, FloatStream) = CreateFloatPropertyStream(*MainTestObject);
 		SenderArgs.Streams = { FloatStream };
 		
 		SetUpClientAndServer();
