Fix comparision with itself

diff --git a/Engine/Plugins/Developer/Concert/ConcertSync/ConcertSyncTest/Source/ConcertSyncTest/Private/Replication/Components/FrequencySettingsTests.cpp b/Engine/Plugins/Developer/Concert/ConcertSync/ConcertSyncTest/Source/ConcertSyncTest/Private/Replication/Components/FrequencySettingsTests.cpp
--- a/Engine/Plugins/Developer/Concert/ConcertSync/ConcertSyncTest/Source/ConcertSyncTest/Private/Replication/Components/FrequencySettingsTests.cpp
+++ b/Engine/Plugins/Developer/Concert/ConcertSync/ConcertSyncTest/Source/ConcertSyncTest/Private/Replication/Components/FrequencySettingsTests.cpp
@@ -3,8 +3,6 @@
 #include "Replication/Data/ReplicationFrequencySettings.h"
 
 #include "Misc/AutomationTest.h"
-#include "Misc/Paths.h"
-#include "HAL/FileManager.h"
 
 namespace UE::ConcertSyncTests
 {
@@ -18,7 +16,6 @@ namespace UE::ConcertSyncTests
 		constexpr FConcertObjectReplicationSettings Realtime_20 { EConcertObjectReplicationMode::Realtime, 20 };
 		constexpr FConcertObjectReplicationSettings Realtime_60 { EConcertObjectReplicationMode::Realtime, 60 };
 
-		TestTrue(TEXT("FPS_30 <= FPS_30 == true"), FPS_30 <= FPS_30);
 		TestTrue(TEXT("FPS_30 < FPS_60 == true"), FPS_30 < FPS_60);
 		TestTrue(TEXT("FPS_30 <= FPS_60 == true"), FPS_30 <= FPS_60);
 		TestFalse(TEXT("FPS_60 < FPS_30 == false"), FPS_60 < FPS_30);
@@ -29,7 +26,6 @@ namespace UE::ConcertSyncTests
 		TestFalse(TEXT("Realtime_20 < FPS_30 == false"), Realtime_20 < FPS_30);
 		TestFalse(TEXT("Realtime_20 <= FPS_30 == false"), Realtime_20 <= FPS_30);
 		
-		TestTrue(TEXT("Realtime_20 <= Realtime_20 == true"), Realtime_20 <= Realtime_20);
 		TestTrue(TEXT("Realtime_20 <= Realtime_60 == true"), Realtime_20 <= Realtime_60);
 		TestTrue(TEXT("Realtime_60 <= Realtime_20 == true"), Realtime_60 <= Realtime_20);
 		TestFalse(TEXT("Realtime_20 < Realtime_60 == false"), Realtime_20 < Realtime_60);
