Add #if WITH_AUTOMATION_TESTS because that's when FWaitLatentCommand is defined

diff --git a/Engine/Plugins/Developer/Concert/ConcertSync/ConcertSyncTest/Source/ConcertSyncTest/Private/Replication/FrequencyTests.cpp b/Engine/Plugins/Developer/Concert/ConcertSync/ConcertSyncTest/Source/ConcertSyncTest/Private/Replication/FrequencyTests.cpp
--- a/Engine/Plugins/Developer/Concert/ConcertSync/ConcertSyncTest/Source/ConcertSyncTest/Private/Replication/FrequencyTests.cpp
+++ b/Engine/Plugins/Developer/Concert/ConcertSync/ConcertSyncTest/Source/ConcertSyncTest/Private/Replication/FrequencyTests.cpp
@@ -8,6 +8,8 @@
 #include "Misc/AutomationTest.h"
 #include "Tests/AutomationCommon.h"
 
+#if WITH_AUTOMATION_TESTS
+
 namespace UE::ConcertSyncTests::Replication::Frequency
 {
 	/** Makes the test object replicate at a frequency of 30 times per second. */
@@ -151,3 +153,5 @@ namespace UE::ConcertSyncTests::Replication::Frequency
 		return true;
 	}
 }
+
+#endif
\ No newline at end of file
