Remove FReplicationStreamDescription_NetPacked

diff --git a/Engine/Plugins/Developer/Concert/ConcertApp/MultiUserClient/Source/MultiUserClient/Private/Replication/Util/StreamRequestUtils.cpp b/Engine/Plugins/Developer/Concert/ConcertApp/MultiUserClient/Source/MultiUserClient/Private/Replication/Util/StreamRequestUtils.cpp
--- a/Engine/Plugins/Developer/Concert/ConcertApp/MultiUserClient/Source/MultiUserClient/Private/Replication/Util/StreamRequestUtils.cpp
+++ b/Engine/Plugins/Developer/Concert/ConcertApp/MultiUserClient/Source/MultiUserClient/Private/Replication/Util/StreamRequestUtils.cpp
@@ -17,7 +17,7 @@ namespace UE::MultiUserClient::StreamRequestUtils
 		
 		FConcertReplication_ChangeStream_Request Request;
 		Request.StreamsToAdd.Emplace();
-		FReplicationStreamDescription_NetPacked& NewStream = Request.StreamsToAdd[0];
+		FReplicationStreamDescription& NewStream = Request.StreamsToAdd[0];
 		NewStream.BaseDescription.Identifier = StreamId;
 			
 		// If creating a new stream, the objects must be supplied in the description instead of in PutObjects!
diff --git a/Engine/Plugins/Developer/Concert/ConcertSync/ConcertSyncClient/Source/ConcertSyncClient/Private/Replication/Manager/ReplicationManagerState_Handshaking.cpp b/Engine/Plugins/Developer/Concert/ConcertSync/ConcertSyncClient/Source/ConcertSyncClient/Private/Replication/Manager/ReplicationManagerState_Handshaking.cpp
--- a/Engine/Plugins/Developer/Concert/ConcertSync/ConcertSyncClient/Source/ConcertSyncClient/Private/Replication/Manager/ReplicationManagerState_Handshaking.cpp
+++ b/Engine/Plugins/Developer/Concert/ConcertSync/ConcertSyncClient/Source/ConcertSyncClient/Private/Replication/Manager/ReplicationManagerState_Handshaking.cpp
@@ -47,9 +47,7 @@ namespace UE::ConcertSyncClient::Replication
 
 	void FReplicationManagerState_Handshaking::OnEnterState()
 	{
-		TArray<FReplicationStreamDescription_NetPacked> Streams_NetPacked;
-		Algo::Transform(RequestArgs.Streams, Streams_NetPacked, [](const FReplicationStreamDescription& Description){ return Description.Pack(); });
-		const FConcertReplication_Join_Request Request{ Streams_NetPacked };
+		const FConcertReplication_Join_Request Request{ RequestArgs.Streams };
 		LiveSession->SendCustomRequest<FConcertReplication_Join_Request, FConcertReplication_Join_Response>(Request, LiveSession->GetSessionServerEndpointId())
 			.Next([WeakThis = TWeakPtr<FReplicationManagerState_Handshaking>(SharedThis(this))](const FConcertReplication_Join_Response& Response)
 			{
diff --git a/Engine/Plugins/Developer/Concert/ConcertSync/ConcertSyncCore/Source/ConcertSyncCore/Private/Replication/ChangeStreamSharedUtils.cpp b/Engine/Plugins/Developer/Concert/ConcertSync/ConcertSyncCore/Source/ConcertSyncCore/Private/Replication/ChangeStreamSharedUtils.cpp
--- a/Engine/Plugins/Developer/Concert/ConcertSync/ConcertSyncCore/Source/ConcertSyncCore/Private/Replication/ChangeStreamSharedUtils.cpp
+++ b/Engine/Plugins/Developer/Concert/ConcertSync/ConcertSyncCore/Source/ConcertSyncCore/Private/Replication/ChangeStreamSharedUtils.cpp
@@ -100,14 +100,10 @@ namespace UE::ConcertSyncCore::Replication::ChangeStreamUtils
 				checkf(ObjectInfo.ClassPath.IsValid(), TEXT("Request did not validate Properties!"));
 			}
 		}
-		
-		for (const FReplicationStreamDescription_NetPacked& StreamDescription : Request.StreamsToAdd)
+
+		for (const FReplicationStreamDescription& StreamDescription : Request.StreamsToAdd)
 		{
-			if (const TOptional<FReplicationStreamDescription> Unpacked = StreamDescription.Unpack()
-				; ensureMsgf(Unpacked, TEXT("Request did not validate unpacking!")))
-			{
-				StreamsToModify.Add(*Unpacked);
-			}
+			StreamsToModify.Add(StreamDescription);
 		}
 	}
 	
diff --git a/Engine/Plugins/Developer/Concert/ConcertSync/ConcertSyncCore/Source/ConcertSyncCore/Private/Replication/Data/ReplicationStreamDescription.cpp b/Engine/Plugins/Developer/Concert/ConcertSync/ConcertSyncCore/Source/ConcertSyncCore/Private/Replication/Data/ReplicationStreamDescription.cpp
deleted file mode 100644
diff --git a/Engine/Plugins/Developer/Concert/ConcertSync/ConcertSyncCore/Source/ConcertSyncCore/Public/Replication/Data/ReplicationStreamDescription.h b/Engine/Plugins/Developer/Concert/ConcertSync/ConcertSyncCore/Source/ConcertSyncCore/Public/Replication/Data/ReplicationStreamDescription.h
--- a/Engine/Plugins/Developer/Concert/ConcertSync/ConcertSyncCore/Source/ConcertSyncCore/Public/Replication/Data/ReplicationStreamDescription.h
+++ b/Engine/Plugins/Developer/Concert/ConcertSync/ConcertSyncCore/Source/ConcertSyncCore/Public/Replication/Data/ReplicationStreamDescription.h
@@ -94,31 +94,6 @@ struct FReplicationStreamDescription
 	CONCERTSYNCCORE_API struct FReplicationStreamDescription_NetPacked Pack() const;
 };
 
-/**
- * This data is packed such that it can be sent via concert messages; packed here means that instanced UObjects are wrapped
- * with sufficient data to reconstruct them.
- */
-USTRUCT()
-struct FReplicationStreamDescription_NetPacked
-{
-	GENERATED_BODY()
-	
-	UPROPERTY()
-	FSharedReplicationStreamDescription BaseDescription;
-
-	friend bool operator==(const FReplicationStreamDescription_NetPacked& Left, const FReplicationStreamDescription_NetPacked& Right)
-	{
-		return Left.BaseDescription == Right.BaseDescription;
-	}
-	friend bool operator!=(const FReplicationStreamDescription_NetPacked& Left, const FReplicationStreamDescription_NetPacked& Right)
-	{
-		return !(Left == Right);
-	}
-
-	/** Unpacks this data. This function can fail, e.g if the UObject class does not exist locally. */
-	CONCERTSYNCCORE_API TOptional<FReplicationStreamDescription> Unpack(FString* OutErrorMessage = nullptr) const;
-};
-
 inline FSharedReplicationStreamDescription FSharedReplicationStreamDescription::Clone(UE::ConcertSyncCore::EReplicationStreamCloneFlags Flags) const
 {
 	using namespace UE::ConcertSyncCore;
diff --git a/Engine/Plugins/Developer/Concert/ConcertSync/ConcertSyncCore/Source/ConcertSyncCore/Public/Replication/Messages/ChangeStream.h b/Engine/Plugins/Developer/Concert/ConcertSync/ConcertSyncCore/Source/ConcertSyncCore/Public/Replication/Messages/ChangeStream.h
--- a/Engine/Plugins/Developer/Concert/ConcertSync/ConcertSyncCore/Source/ConcertSyncCore/Public/Replication/Messages/ChangeStream.h
+++ b/Engine/Plugins/Developer/Concert/ConcertSync/ConcertSyncCore/Source/ConcertSyncCore/Public/Replication/Messages/ChangeStream.h
@@ -172,7 +172,7 @@ struct FConcertReplication_ChangeStream_Request
 	 * in the stream's content being replaced.
 	 */
 	UPROPERTY()
-	TArray<FReplicationStreamDescription_NetPacked> StreamsToAdd;
+	TArray<FReplicationStreamDescription> StreamsToAdd;
 
 	/**
 	 * Streams to remove from the server.
diff --git a/Engine/Plugins/Developer/Concert/ConcertSync/ConcertSyncCore/Source/ConcertSyncCore/Public/Replication/Messages/Handshake.h b/Engine/Plugins/Developer/Concert/ConcertSync/ConcertSyncCore/Source/ConcertSyncCore/Public/Replication/Messages/Handshake.h
--- a/Engine/Plugins/Developer/Concert/ConcertSync/ConcertSyncCore/Source/ConcertSyncCore/Public/Replication/Messages/Handshake.h
+++ b/Engine/Plugins/Developer/Concert/ConcertSync/ConcertSyncCore/Source/ConcertSyncCore/Public/Replication/Messages/Handshake.h
@@ -63,7 +63,7 @@ struct FConcertReplication_Join_Request
 
 	/** The data the client offers to send. */
 	UPROPERTY()
-	TArray<FReplicationStreamDescription_NetPacked> Streams;
+	TArray<FReplicationStreamDescription> Streams;
 };
 
 /**
diff --git a/Engine/Plugins/Developer/Concert/ConcertSync/ConcertSyncServer/Source/ConcertSyncServer/Private/Replication/ConcertServerReplicationManager_ChangeStream.cpp b/Engine/Plugins/Developer/Concert/ConcertSync/ConcertSyncServer/Source/ConcertSyncServer/Private/Replication/ConcertServerReplicationManager_ChangeStream.cpp
--- a/Engine/Plugins/Developer/Concert/ConcertSync/ConcertSyncServer/Source/ConcertSyncServer/Private/Replication/ConcertServerReplicationManager_ChangeStream.cpp
+++ b/Engine/Plugins/Developer/Concert/ConcertSync/ConcertSyncServer/Source/ConcertSyncServer/Private/Replication/ConcertServerReplicationManager_ChangeStream.cpp
@@ -67,7 +67,7 @@ namespace UE::ConcertSyncServer::Replication
 		static void ValidateAddedStreamsAreUnique(const FConcertReplication_ChangeStream_Request& Request, const FConcertReplicationClient& Client, FConcertReplication_ChangeStream_Response& OutResponse)
 		{
 			TSet<FGuid> DuplicateEntryDetection;
-			for (const FReplicationStreamDescription_NetPacked& NewStream : Request.StreamsToAdd)
+			for (const FReplicationStreamDescription& NewStream : Request.StreamsToAdd)
 			{
 				// StreamsToAdd is invalid if there is already a stream with the same ID registered ...
 				const FGuid& NewStreamId = NewStream.BaseDescription.Identifier;
@@ -119,19 +119,6 @@ namespace UE::ConcertSyncServer::Replication
 			}
 		}
 
-		/** Stream description can have UObjects attached which fail to serialize on the receiving end. If so, reject. */
-		static void ValidateUnpacking(const FConcertReplication_ChangeStream_Request& Request, FConcertReplication_ChangeStream_Response& OutResponse)
-		{
-			for (const FReplicationStreamDescription_NetPacked& ToUnpack : Request.StreamsToAdd)
-			{
-				if (!ToUnpack.Unpack())
-				{
-					UE_LOG(LogConcert, Log, TEXT("Failed to unpack stream %s"), *ToUnpack.BaseDescription.Identifier.ToString(EGuidFormats::Short));
-					OutResponse.FailedStreamCreation.Add(ToUnpack.BaseDescription.Identifier);
-				}
-			}
-		}
-
 		/** Checks whether this request is valid to apply. */
 		static bool ShouldAcceptRequest(
 			const FConcertReplication_ChangeStream_Request& Request,
@@ -145,7 +132,6 @@ namespace UE::ConcertSyncServer::Replication
 			ValidatePutObjectsRequestSemantics(Request, Client, OutResponse);
 			ValidateAddedStreamsAreUnique(Request, Client, OutResponse);
 			LookForAuthorityConflicts(Request, Client, AuthorityManager, OutResponse);
-			ValidateUnpacking(Request, OutResponse);
 			ConcertSyncCore::Replication::ChangeStreamUtils::ValidateFrequencyChanges(Request, Client.GetStreamDescriptions(), &OutResponse.FrequencyErrors);
 			
 			return OutResponse.IsSuccess();
diff --git a/Engine/Plugins/Developer/Concert/ConcertSync/ConcertSyncServer/Source/ConcertSyncServer/Private/Replication/Util/JoinRequestValidation.cpp b/Engine/Plugins/Developer/Concert/ConcertSync/ConcertSyncServer/Source/ConcertSyncServer/Private/Replication/Util/JoinRequestValidation.cpp
--- a/Engine/Plugins/Developer/Concert/ConcertSync/ConcertSyncServer/Source/ConcertSyncServer/Private/Replication/Util/JoinRequestValidation.cpp
+++ b/Engine/Plugins/Developer/Concert/ConcertSync/ConcertSyncServer/Source/ConcertSyncServer/Private/Replication/Util/JoinRequestValidation.cpp
@@ -19,7 +19,7 @@ namespace UE::ConcertSyncServer::Replication
 		 */
 		
 		TSet<FGuid> UniqueStreamsFromRequestor;
-		for (const FReplicationStreamDescription_NetPacked& Stream : Request.Streams)
+		for (const FReplicationStreamDescription& Stream : Request.Streams)
 		{
 			const FGuid& StreamId = Stream.BaseDescription.Identifier;
 
@@ -77,26 +77,7 @@ namespace UE::ConcertSyncServer::Replication
 		{
 			return { ErrorCode, ErrorMessage, TArray<FReplicationStreamDescription>{} };
 		}
-		
-		TArray<FReplicationStreamDescription> ParsedStreamDescriptions;
-		ParsedStreamDescriptions.Reserve(Request.Streams.Num());
-		for (const FReplicationStreamDescription_NetPacked& StreamDescription_NetPacked : Request.Streams)
-		{
-			// The net packed version contains UObject attributes, for which we may not be able to look up the UClass for
-			if (TOptional<FReplicationStreamDescription> Description = StreamDescription_NetPacked.Unpack())
-			{
-				ParsedStreamDescriptions.Emplace(MoveTemp(*Description));
-			}
-			else
-			{
-				return {
-					EJoinReplicationErrorCode::FailedToUnpackStream,
-					FString::Printf(TEXT("Failed to unpack stream %s"), *StreamDescription_NetPacked.BaseDescription.Identifier.ToString()),
-					TArray<FReplicationStreamDescription>{}
-				};
-			}
-		}
 
-		return { EJoinReplicationErrorCode::Success, TEXT(""), ParsedStreamDescriptions };
+		return { EJoinReplicationErrorCode::Success, TEXT(""), Request.Streams };
 	}
 }
diff --git a/Engine/Plugins/Developer/Concert/ConcertSync/ConcertSyncTest/Source/ConcertSyncTest/Private/Replication/StreamRequestTests_Change.cpp b/Engine/Plugins/Developer/Concert/ConcertSync/ConcertSyncTest/Source/ConcertSyncTest/Private/Replication/StreamRequestTests_Change.cpp
--- a/Engine/Plugins/Developer/Concert/ConcertSync/ConcertSyncTest/Source/ConcertSyncTest/Private/Replication/StreamRequestTests_Change.cpp
+++ b/Engine/Plugins/Developer/Concert/ConcertSync/ConcertSyncTest/Source/ConcertSyncTest/Private/Replication/StreamRequestTests_Change.cpp
@@ -36,7 +36,7 @@ namespace UE::ConcertSyncTests::Replication::Stream
 		// 2.2 Add new stream
 		auto[DynamicStreamId, DynamicStream] = CreateFloatPropertyStream(*TestObject);
 		FConcertReplication_ChangeStream_Request DynamicCreationRequest;
-		DynamicCreationRequest.StreamsToAdd.Add(DynamicStream.Pack());
+		DynamicCreationRequest.StreamsToAdd.Add(DynamicStream);
 		ChangeStreamForSenderClientAndValidate(TEXT("DynamicCreationRequest"), DynamicCreationRequest, { InitialStream.BaseDescription, DynamicStream.BaseDescription });
 
 		// 2.3 Remove stream
@@ -65,12 +65,12 @@ namespace UE::ConcertSyncTests::Replication::Stream
 		// 2.1 Add stream with float property
 		auto[FloatStreamID, FloatStream] = CreateFloatPropertyStream(*TestObject);
 		FConcertReplication_ChangeStream_Request CreateFloatStreamRequest;
-		CreateFloatStreamRequest.StreamsToAdd.Add(FloatStream.Pack());
+		CreateFloatStreamRequest.StreamsToAdd.Add(FloatStream);
 		ChangeStreamForSenderClientAndValidate(TEXT("CreateFloatStreamRequest"), CreateFloatStreamRequest, { FloatStream.BaseDescription });
 		// 2.2 Add a stream with vector property
 		auto[VectorFloatStreamID, VectorFloatStream] = CreateVectorPropertyStream(*TestObject);
 		FConcertReplication_ChangeStream_Request CreateVectorFloatStreamRequest; 
-		CreateVectorFloatStreamRequest.StreamsToAdd.Add(VectorFloatStream.Pack());
+		CreateVectorFloatStreamRequest.StreamsToAdd.Add(VectorFloatStream);
 		ChangeStreamForSenderClientAndValidate(TEXT("CreateVectorFloatStreamRequest"), CreateVectorFloatStreamRequest, { FloatStream.BaseDescription, VectorFloatStream.BaseDescription });
 
 		// 2.3 Take authority over both streams
@@ -112,7 +112,7 @@ namespace UE::ConcertSyncTests::Replication::Stream
 		// 2.1 Add stream with float property
 		auto[SenderStreamID, SenderStream] = CreateFloatPropertyStream(*TestObject);
 		FConcertReplication_ChangeStream_Request CreateFloatStreamRequest;
-		CreateFloatStreamRequest.StreamsToAdd.Add(SenderStream.Pack());
+		CreateFloatStreamRequest.StreamsToAdd.Add(SenderStream);
 		ChangeStreamForSenderClientAndValidate(TEXT("CreateFloatStreamRequest"), CreateFloatStreamRequest, { SenderStream.BaseDescription });
 
 		// 2.2 Have Sender take authority over the float property
@@ -129,7 +129,7 @@ namespace UE::ConcertSyncTests::Replication::Stream
 		// 2.3 Have Receiver create vector stream
 		auto[ReceiverStreamID, ReceiverStream] = CreateVectorPropertyStream(*TestObject);
 		FConcertReplication_ChangeStream_Request CreateVectorStreamRequest; 
-		CreateVectorStreamRequest.StreamsToAdd.Add(ReceiverStream.Pack());
+		CreateVectorStreamRequest.StreamsToAdd.Add(ReceiverStream);
 		bool bAddedVectorStream = false;
 		ClientReplicationManager_Receiver->ChangeStream(CreateVectorStreamRequest)
 			.Next([this, &bAddedVectorStream](FConcertReplication_ChangeStream_Response&& Response)
@@ -198,7 +198,7 @@ namespace UE::ConcertSyncTests::Replication::Stream
 		// 2.1 Add stream with float property
 		auto[StreamId, Stream] = CreateFloatPropertyStream(*TestObject);
 		FConcertReplication_ChangeStream_Request CreateStreamRequest;
-		CreateStreamRequest.StreamsToAdd.Add(Stream.Pack());
+		CreateStreamRequest.StreamsToAdd.Add(Stream);
 		ChangeStreamForSenderClientAndValidate(TEXT("CreateStreamRequest"), CreateStreamRequest, { Stream.BaseDescription });
 
 		// 2.2 Make a request that will fail and check that no changes were made to the original stream
@@ -206,7 +206,7 @@ namespace UE::ConcertSyncTests::Replication::Stream
 		FConcertPropertySelection NewSelection = GetPropertySelection(Stream, *TestObject);
 		AddFloatProperty(NewSelection);
 		InvalidRequest.ObjectsToPut.Add(FObjectInStreamID{ StreamId, TestObject}, FConcertReplication_ChangeStream_PutObject{ NewSelection });
-		InvalidRequest.StreamsToAdd.Add(Stream.Pack()); // This will make it fail due to pre-existing stream ID
+		InvalidRequest.StreamsToAdd.Add(Stream); // This will make it fail due to pre-existing stream ID
 
 		// Server logs a warning when rejecting - avoid the test being marked with a warning.
 		AddExpectedError(TEXT("Rejecting ChangeStream"));
