Remove FConcertPropertyChain::InternalContainerPropertyValueName.
FConcertPropertyChain will now
- never contain inner properties
- allow listing properties of structs with a custom Serialize function

diff --git a/Engine/Plugins/Developer/Concert/ConcertApp/MultiUserClient/Config/BaseMultiUserClient.ini b/Engine/Plugins/Developer/Concert/ConcertApp/MultiUserClient/Config/BaseMultiUserClient.ini
--- a/Engine/Plugins/Developer/Concert/ConcertApp/MultiUserClient/Config/BaseMultiUserClient.ini
+++ b/Engine/Plugins/Developer/Concert/ConcertApp/MultiUserClient/Config/BaseMultiUserClient.ini
@@ -1,2 +1,2 @@
 [/Script/MultiUserClient.MultiUserReplicationSettings]
-ReplicationEditorSettings=(DefaultPropertySelection=(("/Script/Engine.SceneComponent", (DefaultSelectedProperties=("RelativeLocation","RelativeLocation.X","RelativeLocation.Y","RelativeLocation.Z","RelativeRotation"))),("/Script/CinematicCamera.CineCameraComponent", ()),("/Script/Engine.CameraComponent", ())),DefaultAddedSubobjectRules=(SubobjectMatchingRules=(("/Script/Engine.Actor", (IncludeClasses=("/Script/Engine.StaticMeshComponent","/Script/Engine.CameraComponent","/Script/CinematicCamera.CineCameraComponent"))))))
\ No newline at end of file
+ReplicationEditorSettings=(DefaultPropertySelection=(("/Script/Engine.SceneComponent", (DefaultSelectedProperties=("RelativeLocation","RelativeLocation.X","RelativeLocation.Y","RelativeLocation.Z","RelativeRotation","RelativeRotation.Yaw","RelativeRotation.Pitch","RelativeRotation.Roll"))),("/Script/CinematicCamera.CineCameraComponent", ()),("/Script/Engine.CameraComponent", ())),DefaultAddedSubobjectRules=(SubobjectMatchingRules=(("/Script/Engine.Actor", (IncludeClasses=("/Script/Engine.StaticMeshComponent","/Script/Engine.CameraComponent","/Script/CinematicCamera.CineCameraComponent"))))))
\ No newline at end of file
diff --git a/Engine/Plugins/Developer/Concert/ConcertSync/ConcertSyncCore/Source/ConcertSyncCore/Private/Replication/Data/ConcertPropertySelection.cpp b/Engine/Plugins/Developer/Concert/ConcertSync/ConcertSyncCore/Source/ConcertSyncCore/Private/Replication/Data/ConcertPropertySelection.cpp
--- a/Engine/Plugins/Developer/Concert/ConcertSync/ConcertSyncCore/Source/ConcertSyncCore/Private/Replication/Data/ConcertPropertySelection.cpp
+++ b/Engine/Plugins/Developer/Concert/ConcertSync/ConcertSyncCore/Source/ConcertSyncCore/Private/Replication/Data/ConcertPropertySelection.cpp
@@ -8,8 +8,6 @@
 #include "Algo/AllOf.h"
 #include "Serialization/ArchiveSerializedPropertyChain.h"
 
-const FName FConcertPropertyChain::InternalContainerPropertyValueName(TEXT("Value"));
-
 TOptional<FConcertPropertyChain> FConcertPropertyChain::CreateFromPath(const UStruct& Class, const TArray<FName>& NamePath)
 {
 	TOptional<FConcertPropertyChain> Result;
@@ -31,6 +29,11 @@ FConcertPropertyChain::FConcertPropertyChain(const FArchiveSerializedPropertyCha
 {
 	using namespace UE::ConcertSyncCore::PropertyChain;
 	
+	if (!ensureMsgf(!IsInnerContainerProperty(LeafProperty), TEXT("Concert property chains never contain inner properties so you cannot construct one to point at it!")))
+	{
+		return;
+	}
+	
 	if (OptionalChain)
 	{
 		PathToProperty.Reserve(OptionalChain->GetNumProperties() + 1);
@@ -45,17 +48,12 @@ FConcertPropertyChain::FConcertPropertyChain(const FArchiveSerializedPropertyCha
 				// The input was invalid but we've constructed FConcertPropertyChain with valid state: it ends at the TMap property.
 				return;
 			}
-			
+
+			// Inner properties do not show up in the path but are expected to be contained in OptionalChain since it's needed for the search.
 			if (!IsInnerContainerProperty(*CurrentProperty))
 			{
 				PathToProperty.Add(OptionalChain->GetPropertyFromRoot(i)->GetFName());
 			}
-
-			if (!ensureMsgf(!IsNativeStructProperty(*CurrentProperty), TEXT("Child properties of structs implementing a custom Serialize function never appear in a path. See documentation of PathToProperty!")))
-			{
-				// The input was invalid but we've constructed FConcertPropertyChain with valid state: it ends at the native struct property.
-				return;
-			}
 		}
 	}
 	else if (!ensureMsgf(LeafProperty.GetOwner<FProperty>() == nullptr, TEXT("Not a leaf property!")))
@@ -63,21 +61,8 @@ FConcertPropertyChain::FConcertPropertyChain(const FArchiveSerializedPropertyCha
 		// The input was invalid but we've constructed FConcertPropertyChain with valid state: no property.
 		return;
 	}
-
-	// LeafProperty is allowed to be an inner property of FArrayProperty, FSetProperty, FMapProperty::ValueProp but only if it is primitive or a native structs
-	// In that case the property display label is InternalContainerPropertyValueName, which is "Value"
-	if (IsInnerContainerProperty(LeafProperty))
-	{
-		if (ensureMsgf(IsPrimitiveProperty(LeafProperty) || IsNativeStructProperty(LeafProperty), TEXT("The chain only contains inner properties of primitives or native structs!")))
-		{
-			PathToProperty.Add(InternalContainerPropertyValueName);
-		}
-	}
-	else
-	{
-		PathToProperty.Add(LeafProperty.GetFName());
-	}
-
+	
+	PathToProperty.Add(LeafProperty.GetFName());
 	UE_CLOG(!IsReplicatableProperty(LeafProperty), LogConcert, Warning, TEXT("Instantiated property chain with non replicatable property: %s"), *ToString());
 }
 
@@ -123,11 +108,7 @@ bool FConcertPropertyChain::MatchesExactly(const FArchiveSerializedPropertyChain
 	}
 	
 	const int32 OptionalChainLength = OptionalChain ? OptionalChain->GetNumProperties() : 0;
-	const bool bLeafPropertiesMatch = IsInnerContainerProperty(LeafProperty)
-		// The only place FConcertPropertyChain contains inner container properties is at the end; it is named InternalContainerPropertyValueName.
-		// In that case the "real" leaf property is the owning container property.
-		? (OptionalChainLength >= 1 && PathToProperty.Num() > 1) && OptionalChain->GetPropertyFromStack(0)->GetFName() == PathToProperty[PathToProperty.Num() - 2]
-		: LeafProperty.GetFName() == PathToProperty[PathToProperty.Num() - 1];
+	const bool bLeafPropertiesMatch = LeafProperty.GetFName() == PathToProperty[PathToProperty.Num() - 1];
 	// No point continuing matching properties if leaf properties do not match.
 	if (!bLeafPropertiesMatch)
 	{
@@ -154,9 +135,8 @@ bool FConcertPropertyChain::MatchesExactly(const FArchiveSerializedPropertyChain
 	 *	UCLASS() class UFoo : public UObject { UPROPERTY() TArray<FInner> Nested; }
 	 *	
 	 * Supposing PathToProperty = { "Nested", "Value" }, this is what the input to Matches would be:
-	 *  - OptionalChain = { "Nested", "Nested" }
+	 *  - OptionalChain = { "Nested", "Nested" } ==> 1st Nested is the FArrayProperty, 2nd the FArrayProperty::Inner
 	 *  - LeafProperty = "Value"
-	 * The 1st "Nested" is FArrayProperty and the 2nd "Nested" is the FArrayProperty::Inner property, which is always named the same way.
 	 */	
 	bool bArePathsEqual = true;
 
diff --git a/Engine/Plugins/Developer/Concert/ConcertSync/ConcertSyncCore/Source/ConcertSyncCore/Private/Replication/PropertyChainUtils.cpp b/Engine/Plugins/Developer/Concert/ConcertSync/ConcertSyncCore/Source/ConcertSyncCore/Private/Replication/PropertyChainUtils.cpp
--- a/Engine/Plugins/Developer/Concert/ConcertSync/ConcertSyncCore/Source/ConcertSyncCore/Private/Replication/PropertyChainUtils.cpp
+++ b/Engine/Plugins/Developer/Concert/ConcertSync/ConcertSyncCore/Source/ConcertSyncCore/Private/Replication/PropertyChainUtils.cpp
@@ -55,44 +55,53 @@ namespace UE::ConcertSyncCore::PropertyChain
 			
 			if (FStructProperty* StructProperty = CastField<FStructProperty>(&Property))
 			{
-				// Handle FConcertPropertyChain::InternalContainerPropertyValueName case
-				// If the struct defines a custom serialize function, it does not make sense to list any uproperties we find.
-				// The serialize function can (and often will) skip certain properties so exposing them to a user for selection makes no sense.
-				// If this struct property is supposed to be replicated, then it is all the Serialize function serializes or nothing.
-				return IsNativeStructProperty(*StructProperty)
-					? EBreakBehavior::Continue
-					: VisitStructPropertyRecursive(Chain, *StructProperty, ProcessProperty);
+				/*
+				 * THIS BLOCK OF TEXT CONCERNS NATIVE STRUCTS THAT DEFINE A SERIALIZE() FUNCTION.
+				 * SEE TL;DR.
+				 * 
+				 * Note: If the struct defines a custom serialize function, it *might* not make sense to list any uproperties we find.
+				 * The serialize function can skip certain properties. In those cases  exposing them to a user for selection makes no sense.
+				 * When replicated, effectively the Serialize function decides what is serialized.
+				 * 
+				 * However: Some structs have a Serialize function just return false, which causes standard UPROPERTY serialization.
+				 * Example (in 5.4): FPostProcessSettings.
+				 * 
+				 * The way we'll handle this is by exposing all properties anyway. 
+				 * When replicating the native struct, the struct's Serialize function will then either
+				 * 1. Do its custom logic, which ignores our selection set, or
+				 * 2. Ask the FArchive::ShouldSerializeProperty, which will check the property select set built using this VisitPropertyRecursive function.
+				 * 
+				 * TL;DR: Effectively we'll expose properties (which will be displayed in UI to users) that sometimes cannot be replicated.
+				 * The advantage is that at structs that implement Serialize functions can be correctly replicated. Fair trade-off.
+				 */
+				return VisitStructPropertyRecursive(Chain, *StructProperty, ProcessProperty);
 			}
 
-			auto HandleContainer = [&Chain, &Property, &ProcessProperty](FProperty& Inner)
+			auto HandleContainer = [&Chain, &ProcessProperty](FProperty& Inner)
 			{
-				// Handle FConcertPropertyChain::InternalContainerPropertyValueName case
-				if (IsPrimitiveProperty(Inner)
-					|| IsNativeStructProperty(Inner))
+				// Inner properties are never listed in the path. So we only need to continue if there are more struct sub-properties.
+				FStructProperty* InnerStructProperty = CastField<FStructProperty>(&Inner);
+				if (!InnerStructProperty)
 				{
-					Chain.PushProperty(&Property, Property.IsEditorOnlyProperty());
-					ON_SCOPE_EXIT{ Chain.PopProperty(&Property, Property.IsEditorOnlyProperty()); };
-					return ProcessProperty(Chain, Inner);
+					return EBreakBehavior::Continue;
 				}
 				
-				FStructProperty* InnerStructProperty = CastField<FStructProperty>(&Inner);
 				return InnerStructProperty && IsReplicatableProperty(Inner)
 					? VisitStructPropertyRecursive(Chain, *InnerStructProperty, ProcessProperty)
 					: EBreakBehavior::Continue;
 			};
 
-			if (FArrayProperty* ArrayProperty = CastField<FArrayProperty>(&Property))
+			if (const FArrayProperty* ArrayProperty = CastField<FArrayProperty>(&Property))
 			{
 				return HandleContainer(*ArrayProperty->Inner);
 			}
 
-			if (FSetProperty* SetProperty = CastField<FSetProperty>(&Property))
+			if (const FSetProperty* SetProperty = CastField<FSetProperty>(&Property))
 			{
-				
 				return HandleContainer(*SetProperty->ElementProp);
 			}
 
-			if (FMapProperty* MapProperty = CastField<FMapProperty>(&Property))
+			if (const FMapProperty* MapProperty = CastField<FMapProperty>(&Property))
 			{
 				return IsReplicatableProperty(*MapProperty->KeyProp)
 					? HandleContainer(*MapProperty->ValueProp)
diff --git a/Engine/Plugins/Developer/Concert/ConcertSync/ConcertSyncCore/Source/ConcertSyncCore/Private/Replication/ReplicationPropertyFilter.cpp b/Engine/Plugins/Developer/Concert/ConcertSync/ConcertSyncCore/Source/ConcertSyncCore/Private/Replication/ReplicationPropertyFilter.cpp
--- a/Engine/Plugins/Developer/Concert/ConcertSync/ConcertSyncCore/Source/ConcertSyncCore/Private/Replication/ReplicationPropertyFilter.cpp
+++ b/Engine/Plugins/Developer/Concert/ConcertSync/ConcertSyncCore/Source/ConcertSyncCore/Private/Replication/ReplicationPropertyFilter.cpp
@@ -20,8 +20,14 @@ namespace UE::ConcertSyncCore
 
 	bool FReplicationPropertyFilter::ShouldSerializeProperty(const FArchiveSerializedPropertyChain* Chain, const FProperty& Property) const
 	{
-		// This would mean Property is 1. in a container and 1.1 primitive or 1.2 a native serialized struct 
-		if (PropertyChain::IsPropertyEligibleForMarkingAsInternal(Property))
+		// If Property is in a container then the either it is
+		//  1. primitive, in which case the property is just serialized, or 
+		//  2. a struct that either has
+		//		2.1 No native Serialize function: in this case we'll get recursive ShouldSerializeProperty calls.
+		//		2.2 A native serialize function: In this case we may get 0 or more ShouldSerializeProperty calls.
+		//		The Serialize function will just write whatever it wants. If it returns false, normal UPROPERTY serialization occurs,
+		//		which means we'll end up in case 2.1 
+		if (PropertyChain::IsInnerContainerProperty(Property))
 		{
 			if (ensureMsgf(Chain && Chain->GetNumProperties() >= 1, TEXT("Assumption broken that Property is in a container (array, set, map). Check IsPropertyEligibleForMarkingAsInternal implementation!")))
 			{
diff --git a/Engine/Plugins/Developer/Concert/ConcertSync/ConcertSyncCore/Source/ConcertSyncCore/Public/Replication/Data/ConcertPropertySelection.h b/Engine/Plugins/Developer/Concert/ConcertSync/ConcertSyncCore/Source/ConcertSyncCore/Public/Replication/Data/ConcertPropertySelection.h
--- a/Engine/Plugins/Developer/Concert/ConcertSync/ConcertSyncCore/Source/ConcertSyncCore/Public/Replication/Data/ConcertPropertySelection.h
+++ b/Engine/Plugins/Developer/Concert/ConcertSync/ConcertSyncCore/Source/ConcertSyncCore/Public/Replication/Data/ConcertPropertySelection.h
@@ -10,33 +10,13 @@ struct FArchiveSerializedPropertyChain;
 
 /**
  * Describes the path to a FProperty replicated by Concert.
- * 
- * See Concert.Replication.Data.ForEachReplicatableConcertProperty for path examples (just CTRL+SHIFT+F or go to ConcertSyncTest/Replication/ConcertPropertyTests.cpp).
+ * @see FConcertPropertyChain::PathToProperty.
  */
 USTRUCT()
 struct CONCERTSYNCCORE_API FConcertPropertyChain
 {
 	GENERATED_BODY()
 
-	/**
-	 * The name of inner properties of containers (array, set, map), which are either primitive or native serialized structs.
-	 * 
-	 * Inner properties are included if they are at the end of the path (and only for primitive or native serialized structs!).
-	 * Cases:
-	 * - ArrayOfStructs.ArrayOfFloats.Value
-	 *	There is an array property called ArrayOfStructs containing structs. The contained struct has an array of floats property
-	 *	called ArrayOfFloats. Value corresponds to FArrayProperty::Inner and is called InternalContainerPropertyValueName.
-	 * - ArrayOfStructs.ArrayOfFloats
-	 *  Similar situation but only the size of the array will be replicated. The elements are default initialized.
-	 * - ArrayOfStructs.Value
-	 *  There is an array property called ArrayOfStructs and the struct has a native Serialize() function.
-	 *  
-	 * - BUT not ArrayOfStructs.Value.ArrayOfFloats.Value.
-	 *	This path would imply a different situation in which ArrayOfStructs contains a struct with a struct property called Value;
-	 *	Value's owning struct property would contain an array of floats called ArrayOfFloats.
-	 */
-	static const FName InternalContainerPropertyValueName;
-
 	/** Constructs a FConcertPropertyChain from a path if it is valid. If you need to create many paths in one go, use PropertyUtils::BulkConstructConcertChainsFromPaths instead. */
 	static TOptional<FConcertPropertyChain> CreateFromPath(const UStruct& Class, const TArray<FName>& NamePath);
 
@@ -61,7 +41,12 @@ struct CONCERTSYNCCORE_API FConcertPropertyChain
 	/** @return Whether the leaf property is a direct child of the given property chain. */
 	bool IsDirectChildOf(const FConcertPropertyChain& ParentToCheck) const;
 
-	/** @return Whether OptionalChain and LeafProperty correspond to this path. */
+	/**
+	 * Utility for checking whether this path corresponds to OptionalChain leading to LeafProperty.
+	 * 
+	 * @param OptionalChain Contains all properties leading up to the LeafProperty. You can skip inner container properties.
+	 * @param LeafProperty The last property
+	 * @return Whether OptionalChain and LeafProperty correspond to this path. */
 	bool MatchesExactly(const FArchiveSerializedPropertyChain* OptionalChain, const FProperty& LeafProperty) const;
 
 	/** @return Attempts to resolve this property given the class */
@@ -108,11 +93,22 @@ private:
 	
 	/**
 	 * Path from root of UObject to leaf property. Includes the leaf property.
-	 * This property is kept private to force the use of the exposed constructors.
+	 * Inner container properties, i.e. FArrayProperty::Inner, FSetProperty::ElementProp, FMapProperty::KeyProp, and FMapProperty::ValueProp, are
+	 * never listed in the property path.
 	 *
+	 * Listing some paths (see example code below):
+	 *  - { "Struct", "Foo" }
+	 *  - { "ArrayOfStructs", "Foo" }
+	 *  - { "MapOfStructs" }
+	 *  - { "MapOfStructs", "Foo" } 
 	 * See Concert.Replication.Data.ForEachReplicatableConcertProperty for path examples (just CTRL+SHIFT+F or go to ConcertSyncTest/Replication/ConcertPropertyTests.cpp).
 	 *
-	 * Suppose:
+	 * FConcertPropertyChains do NOT cross the UObject border. In the above example, there would be no such thing as { "UnsupportedInstanced", ... }.
+	 * You'd start a new FConcertPropertyChain for properties for objects of type UFooSubobject.
+	 * 
+	 * This property is kept private to force the use of the exposed constructors.
+	 *
+	 * Code for example:
 	 * class AFooActor
 	 * {
 	 *		UPROPERTY()
@@ -139,35 +135,6 @@ private:
 	 *		UPROPERTY()
 	 *		int32 Bar;
 	 * };
-	 *
-	 * Listing some paths:
-	 *  - { "Struct", "Foo" }
-	 *  - { "ArrayOfStructs", "Foo" }
-	 *  - { "MapOfStructs" } > Replicate all keys but not the struct values
-	 *  - { "MapOfStructs", "Foo" } > Replicate all keys and the Foo struct value.
-	 *	- { "MapIntToFloat" } > Replicate all keys but not the values
-	 *	- { "MapIntToFloat", "Value" } > Replicate all the keys and the values. Special non-ustruct exception where the inner property is listed! 
-	 *  
-	 * The inner properties of containers, e.g. FArrayProperty::Inner, FMapProperty::KeyProp, etc., generally do not appear in the path
-	 * but there is one special case for TMap where it does.
-	 * More details:
-	 *	- TArray, TSet: The inner properties are "obvious" so we can shorten the path.
-	 *	- TMap:
-	 *	  - Key: The key is required. We assume reasonably that all properties of the key property must be replicated. Hence, there is no point in listing
-	 *	  the key inner property nor any of its subproperties. Example: If key is FSoftObjectPath, we do not list FSoftObjectPath::AssetPath, etc.
-	 *	  - Value:
-	 *	    - If the inner value property is is ustruct, then it is not listed. The child properties of the inner property are listed.
-	 *	    - If the inner value property is not a ustruct (meaning a primitive float, int, etc.), then the child property is listed and called "Value" (see InternalContainerPropertyValueName).
-	 *	      This is needed to differentiate between the cases "serializes only the keys" and "serializes the keys and the values".
-	 *	  - Final word about FConcertPropertySelection:
-	 *		- { "MapOfStructs", "Foo" } means all keys and the Foo property is serialized (nothing is said about Bar - it is replicated if it is also in the property selection).
-	 *		- { "MapOfStructs" } means only the keys are serialized and the values are skipped (Foo and Bar will be value-initialized)
-	 *
-	 * Paths do not include the child properties of uStructs implementing a custom Serialize function with struct ops. In that case, the user can only specify
-	 * include the struct or skip it.
-	 *
-	 * FConcertPropertyChains do NOT cross the UObject border. For example, there would be no thing as [1] "UnsupportedInstanced" [x] ...
-	 * You'd start a new FConcertPropertyChain for properties for objects of type UFooSubobject.
 	 */
 	UPROPERTY()
 	TArray<FName> PathToProperty;
diff --git a/Engine/Plugins/Developer/Concert/ConcertSync/ConcertSyncCore/Source/ConcertSyncCore/Public/Replication/PropertyChainUtils.h b/Engine/Plugins/Developer/Concert/ConcertSync/ConcertSyncCore/Source/ConcertSyncCore/Public/Replication/PropertyChainUtils.h
--- a/Engine/Plugins/Developer/Concert/ConcertSync/ConcertSyncCore/Source/ConcertSyncCore/Public/Replication/PropertyChainUtils.h
+++ b/Engine/Plugins/Developer/Concert/ConcertSync/ConcertSyncCore/Source/ConcertSyncCore/Public/Replication/PropertyChainUtils.h
@@ -72,9 +72,6 @@ namespace UE::ConcertSyncCore::PropertyChain
 	/** Util for determining whether this property is the inner property of a container. */
 	CONCERTSYNCCORE_API bool IsInnerContainerProperty(const FProperty& Property);
 
-	/** Checks whether this property would be marked as FConcertPropertyChain::InternalContainerPropertyValueName. */
-	CONCERTSYNCCORE_API bool IsPropertyEligibleForMarkingAsInternal(const FProperty& Property);
-
 	/**
 	 * Util for determining whether this property is considered primitive (numeric, bool, enum).
 	 * Primitive properties in containers are indicated as FConcertPropertyChain::InternalContainerPropertyValueName, which is "Value".
diff --git a/Engine/Plugins/Developer/Concert/ConcertSync/ConcertSyncTest/ConcertSyncTest.uplugin b/Engine/Plugins/Developer/Concert/ConcertSync/ConcertSyncTest/ConcertSyncTest.uplugin
--- a/Engine/Plugins/Developer/Concert/ConcertSync/ConcertSyncTest/ConcertSyncTest.uplugin
+++ b/Engine/Plugins/Developer/Concert/ConcertSync/ConcertSyncTest/ConcertSyncTest.uplugin
@@ -25,6 +25,10 @@
 		}
 	],
 	"Plugins": [
+		{
+			"Name": "ConcertClientSharedSlate",
+			"Enabled": true
+		},
 		{
 			"Name": "ConcertSyncClient",
 			"Enabled": true
diff --git a/Engine/Plugins/Developer/Concert/ConcertSync/ConcertSyncTest/Source/ConcertSyncTest/ConcertSyncTest.Build.cs b/Engine/Plugins/Developer/Concert/ConcertSync/ConcertSyncTest/Source/ConcertSyncTest/ConcertSyncTest.Build.cs
--- a/Engine/Plugins/Developer/Concert/ConcertSync/ConcertSyncTest/Source/ConcertSyncTest/ConcertSyncTest.Build.cs
+++ b/Engine/Plugins/Developer/Concert/ConcertSync/ConcertSyncTest/Source/ConcertSyncTest/ConcertSyncTest.Build.cs
@@ -20,6 +20,7 @@ namespace UnrealBuildTool.Rules
 				new string[]
 				{
 					"Concert",
+					"ConcertClientSharedSlate",
 					"ConcertTransport",
 					"ConcertSyncCore",
 					"ConcertSyncClient",
diff --git a/Engine/Plugins/Developer/Concert/ConcertSync/ConcertSyncTest/Source/ConcertSyncTest/Private/Replication/Components/FrequencySettingsTests.cpp b/Engine/Plugins/Developer/Concert/ConcertSync/ConcertSyncTest/Source/ConcertSyncTest/Private/Replication/Components/FrequencySettingsTests.cpp
--- a/Engine/Plugins/Developer/Concert/ConcertSync/ConcertSyncTest/Source/ConcertSyncTest/Private/Replication/Components/FrequencySettingsTests.cpp
+++ b/Engine/Plugins/Developer/Concert/ConcertSync/ConcertSyncTest/Source/ConcertSyncTest/Private/Replication/Components/FrequencySettingsTests.cpp
@@ -1,7 +1,10 @@
 // Copyright Epic Games, Inc. All Rights Reserved.
 
 #include "Replication/Data/ReplicationFrequencySettings.h"
+#include "Replication/Data/ConcertPropertySelection.h"
+#include "Replication/Settings/ConcertStreamObjectAutoBindingRules.h"
 
+#include "Components/SceneComponent.h"
 #include "Misc/AutomationTest.h"
 
 namespace UE::ConcertSyncTests
@@ -33,4 +36,53 @@ namespace UE::ConcertSyncTests
 		
 		return true;
 	}
+
+	/** Tests that RelativeLocation and RelativeRotation properties are discovered automatically by FConcertStreamObjectAutoBindingRules. */
+	IMPLEMENT_SIMPLE_AUTOMATION_TEST(FAutoBindingPropertySettingsTest, "Editor.Concert.Replication.Components.AutoBindingPropertySettingsTest", EAutomationTestFlags::EditorContext | EAutomationTestFlags::EngineFilter);
+	bool FAutoBindingPropertySettingsTest::RunTest(const FString& Parameters)
+	{
+		// 1. Set up: rules
+		FConcertStreamObjectAutoBindingRules Rules;
+		FConcertDefaultPropertySelection& PropertySelection = Rules.DefaultPropertySelection.Add(USceneComponent::StaticClass());
+		PropertySelection.DefaultSelectedProperties = {
+			TEXT("RelativeLocation"),
+			TEXT("RelativeLocation.X"),
+			TEXT("RelativeLocation.Y"),
+			TEXT("RelativeLocation.Z"),
+			TEXT("RelativeRotation"),
+			TEXT("RelativeRotation.Pitch"),
+			TEXT("RelativeRotation.Yaw"),
+			TEXT("RelativeRotation.Roll")
+		};
+
+		// 2. Run: Apply rules
+		TArray<FConcertPropertyChain> FoundOptions;
+		Rules.AddDefaultPropertiesFromSettings(*USceneComponent::StaticClass(), [&FoundOptions](FConcertPropertyChain&& Chain)
+		{
+			FoundOptions.Emplace(MoveTemp(Chain));
+		});
+
+		// 3. Test
+		TestEqual(TEXT("Num Found == Num Expected"), FoundOptions.Num(), PropertySelection.DefaultSelectedProperties.Num());
+		// Everything that was in the rules was found ...
+		for (const FString& Expected : PropertySelection.DefaultSelectedProperties)
+		{
+			const bool bFound = FoundOptions.ContainsByPredicate([&Expected](const FConcertPropertyChain& PropertyChain)
+			{
+				return Expected == PropertyChain.ToString();
+			});
+			TestTrue(FString::Printf(TEXT("Expected to match %s"), *Expected), bFound);
+		}
+		// ... and did not match anything not specified in the rules
+		for (const FConcertPropertyChain& MatchedChain : FoundOptions)
+		{
+			const bool bFound = PropertySelection.DefaultSelectedProperties.ContainsByPredicate([&MatchedChain](const FString& Expected)
+			{
+				return MatchedChain.ToString() == Expected;
+			});
+			TestTrue(FString::Printf(TEXT("Found %s which was not specified by rules"), *MatchedChain.ToString()), bFound);
+		}
+		
+		return true;
+	}
 }
diff --git a/Engine/Plugins/Developer/Concert/ConcertSync/ConcertSyncTest/Source/ConcertSyncTest/Private/Replication/ConcertPropertyTests.cpp b/Engine/Plugins/Developer/Concert/ConcertSync/ConcertSyncTest/Source/ConcertSyncTest/Private/Replication/ConcertPropertyTests.cpp
--- a/Engine/Plugins/Developer/Concert/ConcertSync/ConcertSyncTest/Source/ConcertSyncTest/Private/Replication/ConcertPropertyTests.cpp
+++ b/Engine/Plugins/Developer/Concert/ConcertSync/ConcertSyncTest/Source/ConcertSyncTest/Private/Replication/ConcertPropertyTests.cpp
@@ -31,26 +31,23 @@ namespace UE::ConcertSyncTest::Replication::PropertyChain
 		FProperty* ValueProp = TestReplicationStruct->FindPropertyByName(GET_MEMBER_NAME_CHECKED(FTestReplicationStruct, Value));
 		FProperty* VectorProp = TestReplicationStruct->FindPropertyByName(GET_MEMBER_NAME_CHECKED(FTestReplicationStruct, Vector));
 		FProperty* VectorXProp = CastField<FStructProperty>(VectorProp)->Struct->FindPropertyByName(GET_MEMBER_NAME_CHECKED(FVector, X));
+		FProperty* NativeStructProp = TestReplicationStruct->FindPropertyByName(GET_MEMBER_NAME_CHECKED(FTestReplicationStruct, NativeStruct));
+		FProperty* NativeStructSubProp = CastField<FStructProperty>(NativeStructProp)->Struct->FindPropertyByName(GET_MEMBER_NAME_CHECKED(FNativeStruct, Float));
 	
 		FProperty* FloatArrayProp = TestReplicationStruct->FindPropertyByName(GET_MEMBER_NAME_CHECKED(FTestReplicationStruct, FloatArray));
-		FProperty* FloatArrayInnerProp = CastField<FArrayProperty>(FloatArrayProp)->Inner;
 		FProperty* StringArrayProp = TestReplicationStruct->FindPropertyByName(GET_MEMBER_NAME_CHECKED(FTestReplicationStruct, StringArray));
 		FProperty* NativeStructArrayProp = TestReplicationStruct->FindPropertyByName(GET_MEMBER_NAME_CHECKED(FTestReplicationStruct, NativeStructArray));
-		FProperty* NativeStructArrayInnerProp = CastField<FArrayProperty>(NativeStructArrayProp)->Inner;
 	
 		FProperty* FloatSetProp = TestReplicationStruct->FindPropertyByName(GET_MEMBER_NAME_CHECKED(FTestReplicationStruct, FloatSet));
-		FProperty* FloatSetInnerProp = CastField<FSetProperty>(FloatSetProp)->ElementProp;
 		FProperty* StringSetProp = TestReplicationStruct->FindPropertyByName(GET_MEMBER_NAME_CHECKED(FTestReplicationStruct, StringSet));
 		FProperty* NativeStructSetProp = TestReplicationStruct->FindPropertyByName(GET_MEMBER_NAME_CHECKED(FTestReplicationStruct, NativeStructSet));
-		FProperty* NativeStructSetInnerProp = CastField<FSetProperty>(NativeStructSetProp)->ElementProp;
 	
 		FProperty* StringToFloatProp = TestReplicationStruct->FindPropertyByName(GET_MEMBER_NAME_CHECKED(FTestReplicationStruct, StringToFloat));
-		FProperty* StringToFloatInnerProp = CastField<FMapProperty>(StringToFloatProp)->ValueProp;
 		FProperty* StringToVectorProp = TestReplicationStruct->FindPropertyByName(GET_MEMBER_NAME_CHECKED(FTestReplicationStruct, StringToVector));
 		FProperty* StringToNativeStructProp = TestReplicationStruct->FindPropertyByName(GET_MEMBER_NAME_CHECKED(FTestReplicationStruct, StringToNativeStruct));
-		FProperty* StringToNativeStructInnerProp = CastField<FMapProperty>(StringToNativeStructProp)->ValueProp;
 
 		FArchiveSerializedPropertyChain VectorXChain;
+		FArchiveSerializedPropertyChain NativeStructChain;
 		FArchiveSerializedPropertyChain FloatArrayChain;
 		FArchiveSerializedPropertyChain FloatSetArrayChain;
 		FArchiveSerializedPropertyChain StringToFloatChain;
@@ -58,6 +55,7 @@ namespace UE::ConcertSyncTest::Replication::PropertyChain
 		FArchiveSerializedPropertyChain NativeStructSetChain;
 		FArchiveSerializedPropertyChain StringToNativeStructChain;
 		VectorXChain.PushProperty(VectorProp, false);
+		NativeStructChain.PushProperty(NativeStructProp, false);
 		FloatArrayChain.PushProperty(FloatArrayProp, false);
 		FloatSetArrayChain.PushProperty(FloatSetProp, false);
 		StringToFloatChain.PushProperty(StringToFloatProp, false);
@@ -71,24 +69,23 @@ namespace UE::ConcertSyncTest::Replication::PropertyChain
 		const FConcertPropertyChain Value(nullptr, *ValueProp);
 		const FConcertPropertyChain Vector(nullptr, *VectorProp);
 		const FConcertPropertyChain VectorX(&VectorXChain, *VectorXProp);
+		const FConcertPropertyChain NativeStruct(nullptr, *NativeStructProp);
+		const FConcertPropertyChain NativeStructSubProperty(&NativeStructChain, *NativeStructSubProp);
 	
 		const FConcertPropertyChain FloatArray(nullptr, *FloatArrayProp);
-		const FConcertPropertyChain FloatArrayInner(&FloatArrayChain, *FloatArrayInnerProp);
 		const FConcertPropertyChain StringArray(nullptr, *StringArrayProp);
 		const FConcertPropertyChain NativeStructArray(nullptr, *NativeStructArrayProp);
-		const FConcertPropertyChain NativeStructArrayInner(&NativeStructArrayChain, *NativeStructArrayInnerProp);
+		const FConcertPropertyChain NativeStructArraySubProperty(&NativeStructArrayChain, *NativeStructSubProp);
 	
 		const FConcertPropertyChain FloatSet(nullptr, *FloatSetProp);
-		const FConcertPropertyChain FloatSetInner(&FloatSetArrayChain, *FloatSetInnerProp);
 		const FConcertPropertyChain StringSet(nullptr, *StringSetProp);
 		const FConcertPropertyChain NativeStructSet(nullptr, *NativeStructSetProp);
-		const FConcertPropertyChain NativeStructSetInner(&NativeStructSetChain, *NativeStructSetInnerProp);
+		const FConcertPropertyChain NativeStructSetSubProperty(&NativeStructSetChain, *NativeStructSubProp);
 	
 		const FConcertPropertyChain StringToFloat(nullptr, *StringToFloatProp);
-		const FConcertPropertyChain StringToFloatInner(&StringToFloatChain, *StringToFloatInnerProp);
 		const FConcertPropertyChain StringToVector(nullptr, *StringToVectorProp);
 		const FConcertPropertyChain StringToNativeStruct(nullptr, *StringToNativeStructProp);
-		const FConcertPropertyChain StringToNativeStructInner(&StringToNativeStructChain, *StringToNativeStructInnerProp);
+		const FConcertPropertyChain StringToNativeStructSubProperty(&StringToNativeStructChain, *NativeStructSubProp);
 
 
 	
@@ -96,24 +93,23 @@ namespace UE::ConcertSyncTest::Replication::PropertyChain
 		TestTrue(TEXT("Value"), Value == TArray<FName>{ TEXT("Value") });
 		TestTrue(TEXT("Vector"), Vector == TArray<FName>{ TEXT("Vector") });
 		TestTrue(TEXT("Vector.X"), VectorX == TArray<FName>{ FName(TEXT("Vector")), FName(TEXT("X")) });
+		TestTrue(TEXT("NativeStruct"), NativeStruct == TArray<FName>{ FName(TEXT("NativeStruct")) });
+		TestTrue(TEXT("NativeStruct.Float"), NativeStructSubProperty == TArray<FName>{ FName(TEXT("NativeStruct")), FName(TEXT("Float")) });
 	
 		TestTrue(TEXT("FloatArray"), FloatArray == TArray<FName>{ TEXT("FloatArray") });
-		TestTrue(TEXT("FloatArrayInner"), FloatArrayInner == TArray<FName>{ FName(TEXT("FloatArray")), FConcertPropertyChain::InternalContainerPropertyValueName });
 		TestTrue(TEXT("StringArray"), StringArray == TArray<FName>{ TEXT("StringArray") });
 		TestTrue(TEXT("NativeStructArray"), NativeStructArray == TArray<FName>{ TEXT("NativeStructArray") });
-		TestTrue(TEXT("NativeStructArrayInner"), NativeStructArrayInner == TArray<FName>{ FName(TEXT("NativeStructArray")), FConcertPropertyChain::InternalContainerPropertyValueName });
+		TestTrue(TEXT("NativeStructArray.Float"), NativeStructArraySubProperty == TArray<FName>{ TEXT("NativeStructArray"), TEXT("Float") });
 	
 		TestTrue(TEXT("FloatSet"), FloatSet == TArray<FName>{ TEXT("FloatSet") });
-		TestTrue(TEXT("FloatSetInner"), FloatSetInner == TArray<FName>{ FName(TEXT("FloatSet")), FConcertPropertyChain::InternalContainerPropertyValueName });
 		TestTrue(TEXT("StringSet"), StringSet == TArray<FName>{ TEXT("StringSet") });
 		TestTrue(TEXT("NativeStructSet"), NativeStructSet == TArray<FName>{ TEXT("NativeStructSet") });
-		TestTrue(TEXT("NativeStructSetInner"), NativeStructSetInner == TArray<FName>{ FName(TEXT("NativeStructSet")), FConcertPropertyChain::InternalContainerPropertyValueName });
+		TestTrue(TEXT("NativeStructSet.Float"), NativeStructSetSubProperty == TArray<FName>{ TEXT("NativeStructSet"), TEXT("Float") });
 	
 		TestTrue(TEXT("StringToFloat"), StringToFloat == TArray<FName>{ TEXT("StringToFloat") });
-		TestTrue(TEXT("StringToFloatInner"), StringToFloatInner == TArray<FName>{ FName(TEXT("StringToFloat")), FConcertPropertyChain::InternalContainerPropertyValueName });
 		TestTrue(TEXT("StringToVector"), StringToVector == TArray<FName>{ TEXT("StringToVector") });
 		TestTrue(TEXT("StringToNativeStruct"), StringToNativeStruct == TArray<FName>{ TEXT("StringToNativeStruct") });
-		TestTrue(TEXT("StringToNativeStructInner"), StringToNativeStructInner == TArray<FName>{ FName(TEXT("StringToNativeStruct")), FConcertPropertyChain::InternalContainerPropertyValueName });
+		TestTrue(TEXT("StringToNativeStruct.Float"), StringToNativeStructSubProperty == TArray<FName>{ TEXT("StringToNativeStruct"), TEXT("Float") });
 	
 		return true;
 	}
@@ -147,24 +143,22 @@ namespace UE::ConcertSyncTest::Replication::PropertyChain
 			{ TEXT("TestStruct"), TEXT("Nested"), TEXT("Vector"), TEXT("Y")  },
 			{ TEXT("TestStruct"), TEXT("Nested"), TEXT("Vector"), TEXT("Z")  },
 			{ TEXT("TestStruct"), TEXT("Nested"), TEXT("NativeStruct") },
+			{ TEXT("TestStruct"), TEXT("Nested"), TEXT("NativeStruct"), TEXT("Float") },
 			{ TEXT("TestStruct"), TEXT("Nested"), TEXT("FloatArray") },
-			{ TEXT("TestStruct"), TEXT("Nested"), TEXT("FloatArray"), FConcertPropertyChain::InternalContainerPropertyValueName },
 			{ TEXT("TestStruct"), TEXT("Nested"), TEXT("StringArray") },
 			{ TEXT("TestStruct"), TEXT("Nested"), TEXT("NativeStructArray") },
-			{ TEXT("TestStruct"), TEXT("Nested"), TEXT("NativeStructArray"), FConcertPropertyChain::InternalContainerPropertyValueName },
+			{ TEXT("TestStruct"), TEXT("Nested"), TEXT("NativeStructArray"), TEXT("Float") },
 			{ TEXT("TestStruct"), TEXT("Nested"), TEXT("FloatSet") },
-			{ TEXT("TestStruct"), TEXT("Nested"), TEXT("FloatSet"), FConcertPropertyChain::InternalContainerPropertyValueName },
 			{ TEXT("TestStruct"), TEXT("Nested"), TEXT("StringSet") },
 			{ TEXT("TestStruct"), TEXT("Nested"), TEXT("NativeStructSet") },
-			{ TEXT("TestStruct"), TEXT("Nested"), TEXT("NativeStructSet"), FConcertPropertyChain::InternalContainerPropertyValueName },
+			{ TEXT("TestStruct"), TEXT("Nested"), TEXT("NativeStructSet"), TEXT("Float") },
 			{ TEXT("TestStruct"), TEXT("Nested"), TEXT("StringToFloat") },
-			{ TEXT("TestStruct"), TEXT("Nested"), TEXT("StringToFloat"), FConcertPropertyChain::InternalContainerPropertyValueName},
 			{ TEXT("TestStruct"), TEXT("Nested"), TEXT("StringToVector") },
 			{ TEXT("TestStruct"), TEXT("Nested"), TEXT("StringToVector"), TEXT("X") },
 			{ TEXT("TestStruct"), TEXT("Nested"), TEXT("StringToVector"), TEXT("Y") },
 			{ TEXT("TestStruct"), TEXT("Nested"), TEXT("StringToVector"), TEXT("Z") },
 			{ TEXT("TestStruct"), TEXT("Nested"), TEXT("StringToNativeStruct") },
-			{ TEXT("TestStruct"), TEXT("Nested"), TEXT("StringToNativeStruct"), FConcertPropertyChain::InternalContainerPropertyValueName },
+			{ TEXT("TestStruct"), TEXT("Nested"), TEXT("StringToNativeStruct"), TEXT("Float") },
 			// The same again for NestedArray as for Nested just above 
 			{ TEXT("TestStruct"), TEXT("NestedArray") },
 			{ TEXT("TestStruct"), TEXT("NestedArray"), TEXT("Value") },
@@ -173,24 +167,22 @@ namespace UE::ConcertSyncTest::Replication::PropertyChain
 			{ TEXT("TestStruct"), TEXT("NestedArray"), TEXT("Vector"), TEXT("Y")  },
 			{ TEXT("TestStruct"), TEXT("NestedArray"), TEXT("Vector"), TEXT("Z")  },
 			{ TEXT("TestStruct"), TEXT("NestedArray"), TEXT("NativeStruct") },
+			{ TEXT("TestStruct"), TEXT("NestedArray"), TEXT("NativeStruct"), TEXT("Float") },
 			{ TEXT("TestStruct"), TEXT("NestedArray"), TEXT("FloatArray") },
-			{ TEXT("TestStruct"), TEXT("NestedArray"), TEXT("FloatArray"), FConcertPropertyChain::InternalContainerPropertyValueName },
 			{ TEXT("TestStruct"), TEXT("NestedArray"), TEXT("StringArray") },
 			{ TEXT("TestStruct"), TEXT("NestedArray"), TEXT("NativeStructArray") },
-			{ TEXT("TestStruct"), TEXT("NestedArray"), TEXT("NativeStructArray"), FConcertPropertyChain::InternalContainerPropertyValueName },
+			{ TEXT("TestStruct"), TEXT("NestedArray"), TEXT("NativeStructArray"), TEXT("Float") },
 			{ TEXT("TestStruct"), TEXT("NestedArray"), TEXT("FloatSet") },
-			{ TEXT("TestStruct"), TEXT("NestedArray"), TEXT("FloatSet"), FConcertPropertyChain::InternalContainerPropertyValueName },
 			{ TEXT("TestStruct"), TEXT("NestedArray"), TEXT("StringSet") },
 			{ TEXT("TestStruct"), TEXT("NestedArray"), TEXT("NativeStructSet") },
-			{ TEXT("TestStruct"), TEXT("NestedArray"), TEXT("NativeStructSet"), FConcertPropertyChain::InternalContainerPropertyValueName },
+			{ TEXT("TestStruct"), TEXT("NestedArray"), TEXT("NativeStructSet"), TEXT("Float") },
 			{ TEXT("TestStruct"), TEXT("NestedArray"), TEXT("StringToFloat") },
-			{ TEXT("TestStruct"), TEXT("NestedArray"), TEXT("StringToFloat"), FConcertPropertyChain::InternalContainerPropertyValueName},
 			{ TEXT("TestStruct"), TEXT("NestedArray"), TEXT("StringToVector") },
 			{ TEXT("TestStruct"), TEXT("NestedArray"), TEXT("StringToVector"), TEXT("X") },
 			{ TEXT("TestStruct"), TEXT("NestedArray"), TEXT("StringToVector"), TEXT("Y") },
 			{ TEXT("TestStruct"), TEXT("NestedArray"), TEXT("StringToVector"), TEXT("Z") },
 			{ TEXT("TestStruct"), TEXT("NestedArray"), TEXT("StringToNativeStruct") },
-			{ TEXT("TestStruct"), TEXT("NestedArray"), TEXT("StringToNativeStruct"), FConcertPropertyChain::InternalContainerPropertyValueName }
+			{ TEXT("TestStruct"), TEXT("NestedArray"), TEXT("StringToNativeStruct"), TEXT("Float") },
 		};
 
 		for (const FExpectedPropertyPath& Expected : ExpectedProperties)
@@ -198,13 +190,18 @@ namespace UE::ConcertSyncTest::Replication::PropertyChain
 			// This structure is easier for setting breakpoints than using TestTrue
 			if (!Properties.Contains(Expected))
 			{
-				AddError(FString::JoinBy(Expected, TEXT("."), [](FName Name){ return Name.ToString(); }));
+				AddError(
+					FString::Printf(
+						TEXT("Expected chain %s"),
+						*FString::JoinBy(Expected, TEXT("."), [](FName Name){ return Name.ToString(); })
+					)
+				);
 			}
 		}
 	
 		if (Properties.Num() != ExpectedProperties.Num())
 		{
-			AddError(TEXT("More properties found than expected"));
+			AddError(TEXT("Different number of properties found than expected"));
 			for (const FConcertPropertyChain& Chain : Properties)
 			{
 				if (!ExpectedProperties.Contains(Chain))
@@ -229,11 +226,8 @@ namespace UE::ConcertSyncTest::Replication::PropertyChain
 		FProperty* TestStructProperty = Class->FindPropertyByName(GET_MEMBER_NAME_CHECKED(UTestReflectionObject, TestStruct));
 		FProperty* NestedArrayProperty = FTestNestedReplicationStruct::StaticStruct()->FindPropertyByName(GET_MEMBER_NAME_CHECKED(FTestNestedReplicationStruct, NestedArray));
 		FProperty* InternalNestedArrayProperty = CastField<FArrayProperty>(NestedArrayProperty)->Inner;
-		FProperty* ValueProperty = FTestReplicationStruct::StaticStruct()->FindPropertyByName(GET_MEMBER_NAME_CHECKED(FTestReplicationStruct, Value));
 		FProperty* FloatArrayProperty = FTestReplicationStruct::StaticStruct()->FindPropertyByName(GET_MEMBER_NAME_CHECKED(FTestReplicationStruct, FloatArray));
-		FProperty* InternalFloatArrayProperty = CastField<FArrayProperty>(FloatArrayProperty)->Inner;
 		FProperty* StringToFloatProperty = FTestReplicationStruct::StaticStruct()->FindPropertyByName(GET_MEMBER_NAME_CHECKED(FTestReplicationStruct, StringToFloat));
-		FProperty* StringToFloatValueProperty = CastField<FMapProperty>(StringToFloatProperty)->ValueProp;
 	
 		FArchiveSerializedPropertyChain ChainToNestedArrayProperty;
 		ChainToNestedArrayProperty.PushProperty(TestStructProperty, false);
@@ -252,22 +246,16 @@ namespace UE::ConcertSyncTest::Replication::PropertyChain
 	
 		const FConcertPropertyChain ConcertChain_Simple_FloatProperty(nullptr, *FloatProperty);
 		const FConcertPropertyChain ConcertChain_Nested_ArrayProperty(&ChainToNestedArrayProperty, *NestedArrayProperty);
-		const FConcertPropertyChain ConcertChain_Nested_ValueProperty(&ChainToNestedProperty, *ValueProperty);
 		const FConcertPropertyChain ConcertChain_Nested_FloatArrayProperty(&ChainToFloatArrayProperty, *FloatArrayProperty);
-		const FConcertPropertyChain ConcertChain_Nested_FloatArrayInnerProperty(&ChainToInternalFloatArrayProperty, *InternalFloatArrayProperty);
 		const FConcertPropertyChain ConcertChain_Nested_StringToFloatProperty(&ChainToStringToFloatProperty, *StringToFloatProperty);
-		const FConcertPropertyChain ConcertChain_Nested_StringToFloatValueProperty(&ChainToStringToFloatValueProperty, *StringToFloatValueProperty);
 	
 
 
 		// 2. Run
 		const bool bMatches_Simple_FloatProperty = ConcertChain_Simple_FloatProperty.MatchesExactly(nullptr, *FloatProperty);
 		const bool bMatches_Nested_ArrayProperty = ConcertChain_Nested_ArrayProperty.MatchesExactly(&ChainToNestedArrayProperty, *NestedArrayProperty);
-		const bool bMatches_Nested_Property = ConcertChain_Nested_ValueProperty.MatchesExactly(&ChainToNestedProperty, *ValueProperty);
 		const bool bMatches_Nested_FloatArrayProperty = ConcertChain_Nested_FloatArrayProperty.MatchesExactly(&ChainToFloatArrayProperty, *FloatArrayProperty);
-		const bool bMatches_Nested_FloatArrayInnerProperty = ConcertChain_Nested_FloatArrayInnerProperty.MatchesExactly(&ChainToInternalFloatArrayProperty, *InternalFloatArrayProperty);
 		const bool bMatches_Nested_StringToFloatProperty = ConcertChain_Nested_StringToFloatProperty.MatchesExactly(&ChainToStringToFloatProperty, *StringToFloatProperty);
-		const bool bMatches_Nested_StringToFloatValueProperty = ConcertChain_Nested_StringToFloatValueProperty.MatchesExactly(&ChainToStringToFloatValueProperty, *StringToFloatValueProperty);
 
 
 
@@ -275,16 +263,8 @@ namespace UE::ConcertSyncTest::Replication::PropertyChain
 		// 3. Test
 		TestTrue(TEXT("Float"), bMatches_Simple_FloatProperty);
 		TestTrue(TEXT("NestedArray"), bMatches_Nested_ArrayProperty);
-		// This case is interesting because e.g. ChainToNestedProperty = { NestedArray, NestedArray } but ConcertChain_NestedProperty = { NestedArray, Value } (set a breakpoint to convince yourself)
-		TestTrue(TEXT("NestedArray.NestedArray.Value"), bMatches_Nested_Property);
-		// This would restore the size of the array only...
-		TestTrue(TEXT("NestedArray.NestedArray.FloatArray"), bMatches_Nested_FloatArrayProperty);
-		// ... and this includes the float values. This one is interesting because it ends with FConcertPropertyChain::InternalContainerPropertyValueName
-		TestTrue(TEXT("NestedArray.NestedArray.FloatArray.Value (ends with FConcertPropertyChain::InternalContainerPropertyValueName)"), bMatches_Nested_FloatArrayInnerProperty);
-		// This would restore the keys and the size of the map only...
-		TestTrue(TEXT("NestedArray.NestedArray.StringToFloat"), bMatches_Nested_StringToFloatProperty);
-		// ... and this includes the float values. This one is interesting because it ends with FConcertPropertyChain::InternalContainerPropertyValueName
-		TestTrue(TEXT("NestedArray.NestedArray.StringToFloat.Value (ends with FConcertPropertyChain::InternalContainerPropertyValueName)"), bMatches_Nested_StringToFloatValueProperty);
+		TestTrue(TEXT("NestedArray.FloatArray"), bMatches_Nested_FloatArrayProperty);
+		TestTrue(TEXT("NestedArray.StringToFloat"), bMatches_Nested_StringToFloatProperty);
 
 		return true;
 	}
@@ -356,8 +336,8 @@ namespace UE::ConcertSyncTest::Replication::PropertyChain
 			FConcertPropertyChain::CreateFromPath(*UTestReflectionObject::StaticClass(), { TEXT("TestStruct"), TEXT("Nested") }).Get({}),
 			FConcertPropertyChain::CreateFromPath(*UTestReflectionObject::StaticClass(), { TEXT("TestStruct"), TEXT("Nested"), TEXT("Value") }).Get({}),
 			FConcertPropertyChain::CreateFromPath(*UTestReflectionObject::StaticClass(), { TEXT("TestStruct"), TEXT("Nested"), TEXT("NativeStruct") }).Get({}),
+			FConcertPropertyChain::CreateFromPath(*UTestReflectionObject::StaticClass(), { TEXT("TestStruct"), TEXT("Nested"), TEXT("NativeStruct") }).Get({}),
 			FConcertPropertyChain::CreateFromPath(*UTestReflectionObject::StaticClass(), { TEXT("TestStruct"), TEXT("Nested"), TEXT("FloatArray") }).Get({}),
-			FConcertPropertyChain::CreateFromPath(*UTestReflectionObject::StaticClass(), { TEXT("TestStruct"), TEXT("Nested"), TEXT("FloatArray"), FConcertPropertyChain::InternalContainerPropertyValueName }).Get({}),
 			FConcertPropertyChain::CreateFromPath(*UTestReflectionObject::StaticClass(), { TEXT("TestStruct"), TEXT("NestedArray") }).Get({}),
 			FConcertPropertyChain::CreateFromPath(*UTestReflectionObject::StaticClass(), { TEXT("TestStruct"), TEXT("NestedArray"), TEXT("Value") }).Get({}),
 			FConcertPropertyChain::CreateFromPath(*UTestReflectionObject::StaticClass(), { TEXT("TestStruct"), TEXT("NestedArray"), TEXT("Vector") }).Get({}),
@@ -366,7 +346,6 @@ namespace UE::ConcertSyncTest::Replication::PropertyChain
 			FConcertPropertyChain::CreateFromPath(*UTestReflectionObject::StaticClass(), { TEXT("TestStruct"), TEXT("NestedArray"), TEXT("Vector"), TEXT("Z") }).Get({}),
 			FConcertPropertyChain::CreateFromPath(*UTestReflectionObject::StaticClass(), { TEXT("TestStruct"), TEXT("NestedArray"), TEXT("NativeStruct") }).Get({}),
 			FConcertPropertyChain::CreateFromPath(*UTestReflectionObject::StaticClass(), { TEXT("TestStruct"), TEXT("NestedArray"), TEXT("FloatArray") }).Get({}),
-			FConcertPropertyChain::CreateFromPath(*UTestReflectionObject::StaticClass(), { TEXT("TestStruct"), TEXT("NestedArray"), TEXT("FloatArray"), FConcertPropertyChain::InternalContainerPropertyValueName }).Get({})
 		};
 		if (AllowedProperties.Contains(FConcertPropertyChain{}))
 		{
