Fix Concert Replication not replicating CDO values

diff --git a/Engine/Plugins/Developer/Concert/ConcertSync/ConcertSyncCore/Source/ConcertSyncCore/Private/Replication/Formats/FullObjectFormat.cpp b/Engine/Plugins/Developer/Concert/ConcertSync/ConcertSyncCore/Source/ConcertSyncCore/Private/Replication/Formats/FullObjectFormat.cpp
--- a/Engine/Plugins/Developer/Concert/ConcertSync/ConcertSyncCore/Source/ConcertSyncCore/Private/Replication/Formats/FullObjectFormat.cpp
+++ b/Engine/Plugins/Developer/Concert/ConcertSync/ConcertSyncCore/Source/ConcertSyncCore/Private/Replication/Formats/FullObjectFormat.cpp
@@ -16,13 +16,14 @@ namespace UE::ConcertSyncCore
 		{
 			Archive.ArIgnoreArchetypeRef = true;
 			Archive.ArIgnoreClassRef = true;
-			Archive.ArNoDelta = false;
+			Archive.ArNoDelta = true; // Always serialize CDO values
+			Archive.SetUseUnversionedPropertySerialization(false); // This is faster - Concert ensures all clients are on the same version
+			// We'd like to use binary serialization but sadly we must used tagged serialization since otherwise the properties cannot be mapped
+			//Archive.SetWantBinaryPropertySerialization(true);
 			
 			Archive.SetIsSaving(!bIsLoading);
 			Archive.SetIsLoading(bIsLoading);
 			Archive.SetIsPersistent(false);
-			// TODO: Test
-			//Archive.SetWantBinaryPropertySerialization(true);
 		}
 		
 		class FFullObjectFormatWriter : public FObjectWriter
diff --git a/Engine/Plugins/Developer/Concert/ConcertSync/ConcertSyncTest/Source/ConcertSyncTest/Private/Replication/ReplicationSendReceiveTests.cpp b/Engine/Plugins/Developer/Concert/ConcertSync/ConcertSyncTest/Source/ConcertSyncTest/Private/Replication/ReplicationSendReceiveTests.cpp
--- a/Engine/Plugins/Developer/Concert/ConcertSync/ConcertSyncTest/Source/ConcertSyncTest/Private/Replication/ReplicationSendReceiveTests.cpp
+++ b/Engine/Plugins/Developer/Concert/ConcertSync/ConcertSyncTest/Source/ConcertSyncTest/Private/Replication/ReplicationSendReceiveTests.cpp
@@ -37,7 +37,11 @@ namespace UE::ConcertSyncTests::Replication::SendReceiveFlow
 	}
 	
 	/** Send values from sender > server > receiver and validate it arrived. */
-	static void SharedRunSendReceiveTest(FSendReceiveObjectTestBase& Test, EPropertyTestFlags PropertyTestFlags = EPropertyTestFlags::All)
+	static void SharedRunSendReceiveTest(FSendReceiveObjectTestBase& Test,
+		EPropertyTestFlags PropertyTestFlags = EPropertyTestFlags::All,
+		// Explicitly sets all values before & after sending so the caller can test whether the other properties were changed
+		EPropertyTestFlags PropertySendFlags = EPropertyTestFlags::All
+		)
 	{
 		bool bHasServerReceivedData = false;
 		bool bHasClientReceivedData = false;
@@ -58,8 +62,8 @@ namespace UE::ConcertSyncTests::Replication::SendReceiveFlow
 			bHasClientReceivedData = true;
 		};
 		Test.SimulateSendObjectToReceiver(OnServerReceive, OnClientReceive,
-			// Explicitly sets all values before & after sending so the caller can test whether the other properties were changed
-			EPropertyTestFlags::All
+			
+			PropertySendFlags
 			);
 
 		// 3. Test
@@ -88,16 +92,23 @@ namespace UE::ConcertSyncTests::Replication::SendReceiveFlow
 		Test.TestTrue(TEXT("Received change stream response"), bReceivedChangeStreamResponse);
 	}
 	
-	/**
-	 * Tests replicating data from sender client > server > receiver client.
-	 */
-	IMPLEMENT_CUSTOM_SIMPLE_AUTOMATION_TEST(FSendReceiveFlowTests, FSendReceiveObjectTestBase, "Editor.Concert.Replication.SendReceive.SingleStream", EAutomationTestFlags::EditorContext | EAutomationTestFlags::EngineFilter);
+	/** Tests replicating data from sender client > server > receiver client. */
+	IMPLEMENT_CUSTOM_SIMPLE_AUTOMATION_TEST(FSendReceiveFlowTests, FSendReceiveObjectTestBase, "Editor.Concert.Replication.SendReceive.SendValuesOnSingleStream", EAutomationTestFlags::EditorContext | EAutomationTestFlags::EngineFilter);
 	bool FSendReceiveFlowTests::RunTest(const FString& Parameters)
 	{
 		SharedSetupSet(*this);
 		SharedRunSendReceiveTest(*this);
 		return true;
 	}
+	
+	/** Tests replicating data from sender client > server > receiver client. */
+	IMPLEMENT_CUSTOM_SIMPLE_AUTOMATION_TEST(FSendCDOValuesOnSingleStream, FSendReceiveObjectTestBase, "Editor.Concert.Replication.SendReceive.SendCDOValuesOnSingleStream", EAutomationTestFlags::EditorContext | EAutomationTestFlags::EngineFilter);
+	bool FSendCDOValuesOnSingleStream::RunTest(const FString& Parameters)
+	{
+		SharedSetupSet(*this);
+		SharedRunSendReceiveTest(*this, EPropertyTestFlags::All | EPropertyTestFlags::SendCDOValues, EPropertyTestFlags::All | EPropertyTestFlags::SendCDOValues);
+		return true;
+	}
 
 	/** Test which still sends TestObject but does so with two separate streams: one for the float and the other for the vector property. */
 	class FSplitSendReceiveObjectTest : public FSendReceiveObjectTestBase
diff --git a/Engine/Plugins/Developer/Concert/ConcertSync/ConcertSyncTest/Source/ConcertSyncTest/Private/Replication/Util/SendReceiveObjectTestBase.cpp b/Engine/Plugins/Developer/Concert/ConcertSync/ConcertSyncTest/Source/ConcertSyncTest/Private/Replication/Util/SendReceiveObjectTestBase.cpp
--- a/Engine/Plugins/Developer/Concert/ConcertSync/ConcertSyncTest/Source/ConcertSyncTest/Private/Replication/Util/SendReceiveObjectTestBase.cpp
+++ b/Engine/Plugins/Developer/Concert/ConcertSync/ConcertSyncTest/Source/ConcertSyncTest/Private/Replication/Util/SendReceiveObjectTestBase.cpp
@@ -97,13 +97,14 @@ namespace UE::ConcertSyncTests::Replication
 
 	void FSendReceiveObjectTestBase::SetTestValues(UTestReflectionObject& Object, EPropertyTestFlags PropertyFlags)
 	{
+		const bool bSendCDOValues = EnumHasAnyFlags(PropertyFlags, EPropertyTestFlags::SendCDOValues);
 		if (EnumHasAnyFlags(PropertyFlags, EPropertyTestFlags::Float))
 		{
-			Object.Float = SentFloat;
+			Object.Float = bSendCDOValues ? GetMutableDefault<UTestReflectionObject>()->Float : SentFloat;
 		}
 		if (EnumHasAnyFlags(PropertyFlags, EPropertyTestFlags::Vector))
 		{
-			Object.Vector = SentVector;
+			Object.Vector = bSendCDOValues ? GetMutableDefault<UTestReflectionObject>()->Vector : SentVector;
 		}
 	}
 	
@@ -121,13 +122,16 @@ namespace UE::ConcertSyncTests::Replication
 	
 	void FSendReceiveObjectTestBase::TestEqualTestValues(UTestReflectionObject& Object, EPropertyTestFlags PropertyFlags)
 	{
+		const bool bSendCDOValues = EnumHasAnyFlags(PropertyFlags, EPropertyTestFlags::SendCDOValues);
 		if (EnumHasAnyFlags(PropertyFlags, EPropertyTestFlags::Float))
 		{
-			TestEqual(TEXT("Float"), Object.Float, SentFloat);
+			const float ExpectedValue = bSendCDOValues ? GetMutableDefault<UTestReflectionObject>()->Float : SentFloat;
+			TestEqual(TEXT("Float"), Object.Float, ExpectedValue);
 		}
 		if (EnumHasAnyFlags(PropertyFlags, EPropertyTestFlags::Vector))
 		{
-			TestEqual(TEXT("Vector"), Object.Vector, SentVector);
+			const FVector ExpectedValue = bSendCDOValues ? GetMutableDefault<UTestReflectionObject>()->Vector : SentVector;
+			TestEqual(TEXT("Vector"), Object.Vector, ExpectedValue);
 		}
 	}
 
diff --git a/Engine/Plugins/Developer/Concert/ConcertSync/ConcertSyncTest/Source/ConcertSyncTest/Private/Replication/Util/SendReceiveObjectTestBase.h b/Engine/Plugins/Developer/Concert/ConcertSync/ConcertSyncTest/Source/ConcertSyncTest/Private/Replication/Util/SendReceiveObjectTestBase.h
--- a/Engine/Plugins/Developer/Concert/ConcertSync/ConcertSyncTest/Source/ConcertSyncTest/Private/Replication/Util/SendReceiveObjectTestBase.h
+++ b/Engine/Plugins/Developer/Concert/ConcertSync/ConcertSyncTest/Source/ConcertSyncTest/Private/Replication/Util/SendReceiveObjectTestBase.h
@@ -24,6 +24,9 @@ namespace UE::ConcertSyncTests::Replication
 		
 		Float = 1 << 0,
 		Vector = 1 << 1,
+
+		/** The CDO values should be sent. */
+		SendCDOValues = 1 << 2,
 		
 		All = Float | Vector
 	};
