MU Replication: Detect when an object is renamed and re-resolve the soft object path
- This handles replicating e.g. a component that is replaced by running the construction script

diff --git a/Engine/Plugins/Developer/Concert/ConcertSync/ConcertSyncClient/Source/ConcertSyncClient/Private/Replication/ConcertClientReplicationBridge.cpp b/Engine/Plugins/Developer/Concert/ConcertSync/ConcertSyncClient/Source/ConcertSyncClient/Private/Replication/ConcertClientReplicationBridge.cpp
--- a/Engine/Plugins/Developer/Concert/ConcertSync/ConcertSyncClient/Source/ConcertSyncClient/Private/Replication/ConcertClientReplicationBridge.cpp
+++ b/Engine/Plugins/Developer/Concert/ConcertSync/ConcertSyncClient/Source/ConcertSyncClient/Private/Replication/ConcertClientReplicationBridge.cpp
@@ -96,9 +96,36 @@ namespace UE::ConcertSyncClient::Replication
 
 	UObject* FConcertClientReplicationBridge::FindObjectIfAvailable(const FSoftObjectPath& Path)
 	{
+		// Replicated UObject's may be renamed, e.g. by using UObject::Rename. If that happens, we need to re-resolve the Path.
+		// A legitimate case that could cause this: We're replicating a component that has EComponentCreationMethod::SimpleConstructionScript.
+		// Upon editing the object, the construction script is re-run thus replacing the component instance with a totally new object!
+		FTrackedObjectInfo* ObjectInfo = TrackedObjects.Find(Path);
+		const bool bObjectWasInCache = ObjectInfo && ObjectInfo->ResolvedObject.IsValid();
+		if (bObjectWasInCache)
+		{
+			if (Path == ObjectInfo->ResolvedObject.Get())
+			{
+				return ObjectInfo->ResolvedObject.Get();
+			}
+			UE_LOG(LogConcert, Verbose, TEXT("Replicated object %s was renamed to %s. Resolving new target object..."), *Path.ToString(), *ObjectInfo->ResolvedObject->GetPathName());
+		}
+		
 		// This should be an object in a UWorld. This will resolve if the world is opened.
 		UObject* Object = Path.ResolveObject();
 		UE_CLOG(Object && !Object->IsInA(UWorld::StaticClass()), LogConcert, Warning, TEXT("Object %s is not in any UWorld. The replication system was designed to work with objects in UWorlds."), *Path.ToString());
+
+		// Update the cache for faster look ups if we found anything
+		if (Object)
+		{
+			ObjectInfo = ObjectInfo ? ObjectInfo : &TrackedObjects.Add(Path);
+			ObjectInfo->ResolvedObject = Object;
+		}
+		else
+		{
+			// If the object was renamed but not replaced, issue a warning.
+			UE_CLOG(bObjectWasInCache, LogConcert, Warning, TEXT("Replicated object %s was renamed to %s. The target object was not replaced by any new object so it will stop being replicated."), *Path.ToString(), *ObjectInfo->ResolvedObject.Get()->GetPathName());
+		}
+		
 		return Object;
 	}
 
diff --git a/Engine/Plugins/Developer/Concert/ConcertSync/ConcertSyncClient/Source/ConcertSyncClient/Private/Replication/Processing/ClientReplicationDataCollector.cpp b/Engine/Plugins/Developer/Concert/ConcertSync/ConcertSyncClient/Source/ConcertSyncClient/Private/Replication/Processing/ClientReplicationDataCollector.cpp
--- a/Engine/Plugins/Developer/Concert/ConcertSync/ConcertSyncClient/Source/ConcertSyncClient/Private/Replication/Processing/ClientReplicationDataCollector.cpp
+++ b/Engine/Plugins/Developer/Concert/ConcertSync/ConcertSyncClient/Source/ConcertSyncClient/Private/Replication/Processing/ClientReplicationDataCollector.cpp
@@ -161,10 +161,17 @@ namespace UE::ConcertSyncClient::Replication
 		{
 			for (const FObjectInfo& ObjectInfo : Pair.Value)
 			{
-				if (ObjectInfo.ObjectCache.IsValid())
+				// The bridge has the object cached. We do not cache the object ourselves!
+				// Funky Unreal flows can cause the object to be renamed out from under us and replaced by a different instance (just by using UObject::Rename()).
+				// The bridge is aware of these flows and FindObjectIfAvailable will catch them.
+				// If FindObjectIfAvailable fails to resolve, OnObjectHidden() is triggered if the object was previously visible.
+				const UObject* Object = Bridge->FindObjectIfAvailable(Pair.Key);
+				
+				const FSoftObjectPath ObjectPath = Object;
+				if (Object && ensureMsgf(ObjectPath == Object, TEXT("Sanity check: the bridge gave us an object with a different path!")))
 				{
 					ProcessItemFunc(FConcertReplicatedObjectId {
-						FConcertObjectInStreamID{ ObjectInfo.StreamId, ObjectInfo.ObjectCache.Get() },
+						FConcertObjectInStreamID{ ObjectInfo.StreamId, Object },
 						ClientId
 					});
 				}
@@ -185,23 +192,26 @@ namespace UE::ConcertSyncClient::Replication
 			return {};
 		}
 
-		const int32 Index = ObjectInfos->IndexOfByPredicate([&ObjectToProcess](const FObjectInfo& ObjectInfo){ return ObjectInfo.StreamId == ObjectToProcess.StreamId; });
-		// Same logic as above - either invalid call or ForEachPendingObject lied
-		if (!ensure(ObjectInfos->IsValidIndex(Index)))
+		// The properties to be replicated depend on the stream: search for it.
+		const int32 StreamIndex = ObjectInfos->IndexOfByPredicate([&ObjectToProcess](const FObjectInfo& ObjectInfo){ return ObjectInfo.StreamId == ObjectToProcess.StreamId; });
+		// Same logic as above: if the stream is not found, then either the call is invalid or ForEachPendingObject lied
+		if (!ensure(ObjectInfos->IsValidIndex(StreamIndex)))
 		{
 			return {};
 		}
 
-		const FObjectInfo& ObjectInfo = (*ObjectInfos)[Index];
-		// Finally... same logic as above - either invalid call or ForEachPendingObject lied
-		if (!ensure(ObjectInfo.ObjectCache.IsValid()))
+		UObject* Object = Bridge->FindObjectIfAvailable(ObjectToProcess.Object);
+		// Finally... ask the bridge to resolve the object for us. The bridge has the object cached and handles the object getting renamed, etc.
+		// This should resolve. If it does not: same logic as above - either this call is invalid or ForEachPendingObject lied.
+		if (!ensure(Object))
 		{
 			return false;
 		}
 
+		const FObjectInfo& ObjectInfo = (*ObjectInfos)[StreamIndex];
 		ConcertSyncCore::FReplicationPropertyFilter Filter(ObjectInfo.SelectedProperties);
 		TOptional<FConcertSessionSerializedPayload> Payload = ReplicationFormat->CreateReplicationEvent(
-			*ObjectInfo.ObjectCache,
+			*Object,
 			[&Filter](const FArchiveSerializedPropertyChain* Chain, const FProperty& Property)
 			{
 				return Filter.ShouldSerializeProperty(Chain, Property);
@@ -218,10 +228,6 @@ namespace UE::ConcertSyncClient::Replication
 	{
 		if (TArray<FObjectInfo>* ObjectInfos = ObjectsToReplicate.Find(&Object))
 		{
-			for (FObjectInfo& ObjectInfo : *ObjectInfos)
-			{
-				ObjectInfo.ObjectCache = &Object;
-			}
 			++NumTrackedObjects;
 		}
 	}
@@ -230,10 +236,6 @@ namespace UE::ConcertSyncClient::Replication
 	{
 		if (TArray<FObjectInfo>* ObjectInfos = ObjectsToReplicate.Find(ObjectPath))
 		{
-			for (FObjectInfo& ObjectInfo : *ObjectInfos)
-			{
-				ObjectInfo.ObjectCache = nullptr;
-			}
 			--NumTrackedObjects;
 		}
 	}
diff --git a/Engine/Plugins/Developer/Concert/ConcertSync/ConcertSyncClient/Source/ConcertSyncClient/Private/Replication/Processing/ClientReplicationDataCollector.h b/Engine/Plugins/Developer/Concert/ConcertSync/ConcertSyncClient/Source/ConcertSyncClient/Private/Replication/Processing/ClientReplicationDataCollector.h
--- a/Engine/Plugins/Developer/Concert/ConcertSync/ConcertSyncClient/Source/ConcertSyncClient/Private/Replication/Processing/ClientReplicationDataCollector.h
+++ b/Engine/Plugins/Developer/Concert/ConcertSync/ConcertSyncClient/Source/ConcertSyncClient/Private/Replication/Processing/ClientReplicationDataCollector.h
@@ -99,8 +99,6 @@ namespace UE::ConcertSyncClient::Replication
 			FGuid StreamId;
 			/** The properties to replicate */
 			FConcertPropertySelection SelectedProperties;
-			/** Set when the bridge tells us the object is available for replication. Unset otherwise. */
-			TWeakObjectPtr<UObject> ObjectCache;
 		};
 		
 		/** The objects and their properties to replicate */
diff --git a/Engine/Plugins/Developer/Concert/ConcertSync/ConcertSyncClient/Source/ConcertSyncClient/Private/Replication/ReplicationTestInterface.cpp b/Engine/Plugins/Developer/Concert/ConcertSync/ConcertSyncClient/Source/ConcertSyncClient/Private/Replication/ReplicationTestInterface.cpp
--- a/Engine/Plugins/Developer/Concert/ConcertSync/ConcertSyncClient/Source/ConcertSyncClient/Private/Replication/ReplicationTestInterface.cpp
+++ b/Engine/Plugins/Developer/Concert/ConcertSync/ConcertSyncClient/Source/ConcertSyncClient/Private/Replication/ReplicationTestInterface.cpp
@@ -1,5 +1,6 @@
 // Copyright Epic Games, Inc. All Rights Reserved.
 
+#include "ConcertClientReplicationBridge.h"
 #include "Replication/IConcertClientReplicationBridge.h"
 #include "Replication/IConcertClientReplicationManager.h"
 #include "Manager/ReplicationManager.h"
@@ -21,4 +22,9 @@ namespace UE::ConcertSyncClient::TestInterface
 		Result->StartAcceptingJoinRequests();
 		return Result;
 	}
+
+	CONCERTSYNCCLIENT_API TSharedRef<IConcertClientReplicationBridge> CreateClientReplicationBridge()
+	{
+		return MakeShared<Replication::FConcertClientReplicationBridge>();
+	}
 }
diff --git a/Engine/Plugins/Developer/Concert/ConcertSync/ConcertSyncClient/Source/ConcertSyncClient/Public/Replication/IConcertClientReplicationBridge.h b/Engine/Plugins/Developer/Concert/ConcertSync/ConcertSyncClient/Source/ConcertSyncClient/Public/Replication/IConcertClientReplicationBridge.h
--- a/Engine/Plugins/Developer/Concert/ConcertSync/ConcertSyncClient/Source/ConcertSyncClient/Public/Replication/IConcertClientReplicationBridge.h
+++ b/Engine/Plugins/Developer/Concert/ConcertSync/ConcertSyncClient/Source/ConcertSyncClient/Public/Replication/IConcertClientReplicationBridge.h
@@ -65,7 +65,7 @@ public:
 	virtual bool IsObjectAvailable(const FSoftObjectPath& Path) = 0;
 
 	/**
-	 * Obtains the object if it is available. If the UObject is unloaded, it shall be loaded.
+	 * Obtains the object if it is available. This does not load the object. 
 	 * 
 	 * This function will not discover nor hide any UObject (that means broadcast OnObjectDiscovered or OnObjectHidden),
 	 * unless it is tracked, in which case it can.
diff --git a/Engine/Plugins/Developer/Concert/ConcertSync/ConcertSyncTest/Source/ConcertSyncTest/Private/Replication/Bugs/ChangeObjectNameDuringReplication.cpp b/Engine/Plugins/Developer/Concert/ConcertSync/ConcertSyncTest/Source/ConcertSyncTest/Private/Replication/Bugs/ChangeObjectNameDuringReplication.cpp
new file mode 100644
--- /dev/null
+++ b/Engine/Plugins/Developer/Concert/ConcertSync/ConcertSyncTest/Source/ConcertSyncTest/Private/Replication/Bugs/ChangeObjectNameDuringReplication.cpp
@@ -0,0 +1,114 @@
+// Copyright Epic Games, Inc. All Rights Reserved.
+
+#include "Replication/Util/SendReceiveWorldTestBase.h"
+
+#include "Components/StaticMeshComponent.h"
+#include "Engine/StaticMeshActor.h"
+#include "Engine/World.h"
+#include "Misc/AutomationTest.h"
+#include "Misc/Paths.h"
+#include "UObject/Package.h"
+
+namespace UE::ConcertSyncTests::Replication::Bugs
+{
+	/**
+	 * This tests that the replication system detects when a replicated object is renamed and subsequently replaced by an equally named object.
+	 * 
+	 * This happens for example when replicating a construction script created component on a Blueprint actor and the actor runs a construction script,
+	 * which reinstances the components.
+	 */
+	IMPLEMENT_CUSTOM_SIMPLE_AUTOMATION_TEST(FChangeObjectNameDuringReplication, FSendReceiveWorldTestBase, "Editor.Concert.Replication.Bugs.ChangeObjectNameDuringReplication", EAutomationTestFlags::EditorContext | EAutomationTestFlags::EngineFilter);
+	bool FChangeObjectNameDuringReplication::RunTest(const FString& Parameters)
+	{
+		// The real replication bridge must be used because we're testing that the real bridge detects renaming
+		SetUpClientAndServer(ESendReceiveTestFlags::UseRealReplicationBridge);
+		SetupWorld();
+
+		const FVector InitialLocation { 10.f };
+		const FVector GarbageLocation { -10.f };
+		const FVector AfterRenameLocation { 20.f };
+		const FName ComponentName = TEXT("ReplicatedComponent");
+
+		// Create an actor a component called "ReplicatedComponent"
+		UWorld* World = GetWorld();
+		AStaticMeshActor* StaticMeshActor = World->SpawnActor<AStaticMeshActor>();
+		USceneComponent* FirstComponent = NewObject<USceneComponent>(StaticMeshActor, ComponentName);
+		FirstComponent->SetupAttachment(StaticMeshActor->GetRootComponent());
+		FirstComponent->RegisterComponent();
+		FirstComponent->SetRelativeLocation(InitialLocation);
+
+		// Validate properties exist (in case FConcertPropertyChain implementation is broken)
+		const TOptional<FConcertPropertyChain> RelativeLocation = FConcertPropertyChain::CreateFromPath(*USceneComponent::StaticClass(), { TEXT("RelativeLocation") });
+		const TOptional<FConcertPropertyChain> RelativeLocationX = FConcertPropertyChain::CreateFromPath(*USceneComponent::StaticClass(), { TEXT("RelativeLocation"), TEXT("X") });
+		const TOptional<FConcertPropertyChain> RelativeLocationY = FConcertPropertyChain::CreateFromPath(*USceneComponent::StaticClass(), { TEXT("RelativeLocation"), TEXT("Y") });
+		const TOptional<FConcertPropertyChain> RelativeLocationZ = FConcertPropertyChain::CreateFromPath(*USceneComponent::StaticClass(), { TEXT("RelativeLocation"), TEXT("Z") });
+		if (!RelativeLocation || !RelativeLocationX || !RelativeLocationY || !RelativeLocationZ)
+		{
+			AddError(TEXT("Failed to find RelativeLocation properties. Check FConcertPropertyChain implementation."));
+			return false;
+		}
+		
+		// Register the component for replication
+		FConcertReplication_ChangeStream_Request StreamRequest;
+		FConcertReplicationStream Stream { FGuid::NewGuid() };
+		Stream.BaseDescription.FrequencySettings.Defaults.ReplicationMode = EConcertObjectReplicationMode::Realtime;
+		FConcertReplicatedObjectInfo& ObjectInfo = Stream.BaseDescription.ReplicationMap.ReplicatedObjects.Add(FirstComponent);
+		ObjectInfo.ClassPath = FirstComponent->GetClass();
+		ObjectInfo.PropertySelection.ReplicatedProperties.Append({ *RelativeLocation, *RelativeLocationX, *RelativeLocationY, *RelativeLocationZ });
+		StreamRequest.StreamsToAdd = { Stream };
+		ClientReplicationManager_Sender->ChangeStream(StreamRequest)
+			.Next([this](FConcertReplication_ChangeStream_Response&& Response)
+			{
+				TestTrue(TEXT("Register component"), Response.IsSuccess());
+			});
+		ClientReplicationManager_Sender->TakeAuthorityOver({ FirstComponent })
+			.Next([this](FConcertReplication_ChangeAuthority_Response&& Response)
+			{
+				TestTrue(TEXT("Take authority of component"), Response.RejectedObjects.IsEmpty());
+			});
+
+		// Replicate the original component
+		{
+			auto SetTestValues = [&](UWorld&)
+			{
+				FirstComponent->SetRelativeLocation(InitialLocation);
+			};
+			auto SetGarbageValues = [&](UWorld&)
+			{
+				FirstComponent->SetRelativeLocation(GarbageLocation);
+			};
+			auto TestValues = [&](UWorld&)
+			{
+				TestEqual(TEXT("Replicate original component"), FirstComponent->GetRelativeLocation(), InitialLocation);
+			};
+			SimulateSenderToReceiverForWorld(SetTestValues, SetGarbageValues, { TestValues });
+		}
+
+		// Rename the component and replace it. This simulates running the construction script.
+		FirstComponent->UnregisterComponent();
+		StaticMeshActor->RemoveOwnedComponent(FirstComponent);
+
+		FirstComponent->Rename(*FString::Printf(TEXT("TestGarbage_%s"), *FGuid::NewGuid().ToString()), GetTransientPackage());
+		USceneComponent* NewComponent = NewObject<USceneComponent>(StaticMeshActor, ComponentName);
+		NewComponent->RegisterComponent();
+
+		// Replicate the new component with same name
+		{
+			auto SetTestValues = [&](UWorld&)
+			{
+				NewComponent->SetRelativeLocation(AfterRenameLocation);
+			};
+			auto SetGarbageValues = [&](UWorld&)
+			{
+				NewComponent->SetRelativeLocation(GarbageLocation);
+			};
+			auto TestValues = [&](UWorld&)
+			{
+				TestEqual(TEXT("Replicate new component"), NewComponent->GetRelativeLocation(), AfterRenameLocation);
+			};
+			SimulateSenderToReceiverForWorld(SetTestValues, SetGarbageValues, { TestValues });
+		}
+		
+		return true;
+	}
+}
diff --git a/Engine/Plugins/Developer/Concert/ConcertSync/ConcertSyncTest/Source/ConcertSyncTest/Private/Replication/ReplicationTestInterface.h b/Engine/Plugins/Developer/Concert/ConcertSync/ConcertSyncTest/Source/ConcertSyncTest/Private/Replication/ReplicationTestInterface.h
--- a/Engine/Plugins/Developer/Concert/ConcertSync/ConcertSyncTest/Source/ConcertSyncTest/Private/Replication/ReplicationTestInterface.h
+++ b/Engine/Plugins/Developer/Concert/ConcertSync/ConcertSyncTest/Source/ConcertSyncTest/Private/Replication/ReplicationTestInterface.h
@@ -20,6 +20,8 @@ namespace UE::ConcertSyncClient::TestInterface
 		TSharedRef<IConcertClientSession> InLiveSession,
 		IConcertClientReplicationBridge* InBridge
 		);
+
+	extern CONCERTSYNCCLIENT_API TSharedRef<IConcertClientReplicationBridge> CreateClientReplicationBridge();
 }
 
 namespace UE::ConcertSyncServer::TestInterface
diff --git a/Engine/Plugins/Developer/Concert/ConcertSync/ConcertSyncTest/Source/ConcertSyncTest/Private/Replication/Util/SendReceiveObjectTestBase.cpp b/Engine/Plugins/Developer/Concert/ConcertSync/ConcertSyncTest/Source/ConcertSyncTest/Private/Replication/Util/SendReceiveObjectTestBase.cpp
--- a/Engine/Plugins/Developer/Concert/ConcertSync/ConcertSyncTest/Source/ConcertSyncTest/Private/Replication/Util/SendReceiveObjectTestBase.cpp
+++ b/Engine/Plugins/Developer/Concert/ConcertSync/ConcertSyncTest/Source/ConcertSyncTest/Private/Replication/Util/SendReceiveObjectTestBase.cpp
@@ -28,12 +28,12 @@ namespace UE::ConcertSyncTests::Replication
 		return ReceiverJoinArgs;
 	}
 
-	void FSendReceiveObjectTestBase::SetUpClientAndServer()
+	void FSendReceiveObjectTestBase::SetUpClientAndServer(ESendReceiveTestFlags Flags)
 	{
 		// Fake replicated object must be created before call to Super
 		TestObject = NewObject<UTestReflectionObject>(GetTransientPackage());
 
-		FSendReceiveTestBase::SetUpClientAndServer();
+		FSendReceiveTestBase::SetUpClientAndServer(Flags);
 		
 		// Bridge is responsible for telling client-side replication system about existing objects
 		BridgeMock_Sender->InjectAvailableObject(*TestObject); 
diff --git a/Engine/Plugins/Developer/Concert/ConcertSync/ConcertSyncTest/Source/ConcertSyncTest/Private/Replication/Util/SendReceiveObjectTestBase.h b/Engine/Plugins/Developer/Concert/ConcertSync/ConcertSyncTest/Source/ConcertSyncTest/Private/Replication/Util/SendReceiveObjectTestBase.h
--- a/Engine/Plugins/Developer/Concert/ConcertSync/ConcertSyncTest/Source/ConcertSyncTest/Private/Replication/Util/SendReceiveObjectTestBase.h
+++ b/Engine/Plugins/Developer/Concert/ConcertSync/ConcertSyncTest/Source/ConcertSyncTest/Private/Replication/Util/SendReceiveObjectTestBase.h
@@ -56,7 +56,7 @@ namespace UE::ConcertSyncTests::Replication
 		//~ Begin FSendReceiveTestBase Interface
 		virtual ConcertSyncClient::Replication::FJoinReplicatedSessionArgs CreateSenderArgs() override;
 		virtual ConcertSyncClient::Replication::FJoinReplicatedSessionArgs CreateReceiverArgs() override;
-		virtual void SetUpClientAndServer() override;
+		virtual void SetUpClientAndServer(ESendReceiveTestFlags Flags = ESendReceiveTestFlags::None) override;
 		//~ End FSendReceiveTestBase Interface
 		
 		/**
diff --git a/Engine/Plugins/Developer/Concert/ConcertSync/ConcertSyncTest/Source/ConcertSyncTest/Private/Replication/Util/SendReceiveTestBase.cpp b/Engine/Plugins/Developer/Concert/ConcertSync/ConcertSyncTest/Source/ConcertSyncTest/Private/Replication/Util/SendReceiveTestBase.cpp
--- a/Engine/Plugins/Developer/Concert/ConcertSync/ConcertSyncTest/Source/ConcertSyncTest/Private/Replication/Util/SendReceiveTestBase.cpp
+++ b/Engine/Plugins/Developer/Concert/ConcertSync/ConcertSyncTest/Source/ConcertSyncTest/Private/Replication/Util/SendReceiveTestBase.cpp
@@ -46,7 +46,7 @@ namespace UE::ConcertSyncTests::Replication
 		return SenderJoinArgs;
 	}
 
-	void FSendReceiveTestBase::SetUpClientAndServer()
+	void FSendReceiveTestBase::SetUpClientAndServer(ESendReceiveTestFlags Flags)
 	{
 		// Server
 		InitServer();
@@ -54,12 +54,29 @@ namespace UE::ConcertSyncTests::Replication
 		ServerReplicationManager = ConcertSyncServer::TestInterface::CreateServerReplicationManager(ServerSession.ToSharedRef());
 		// Client Receiver
 		Client_Receiver = &ConnectClient();
-		BridgeMock_Receiver = MakeShared<FConcertClientReplicationBridgeMock>();
-		ClientReplicationManager_Receiver = ConcertSyncClient::TestInterface::CreateClientReplicationManager(Client_Receiver->ClientSessionMock, BridgeMock_Receiver.Get());
+		if (EnumHasAnyFlags(Flags, ESendReceiveTestFlags::UseRealReplicationBridge))
+		{
+			BridgeUsed_Receiver = ConcertSyncClient::TestInterface::CreateClientReplicationBridge();
+		}
+		else
+		{
+			BridgeMock_Receiver =  MakeShared<FConcertClientReplicationBridgeMock>();
+			BridgeUsed_Receiver = BridgeMock_Receiver;
+		}
+		ClientReplicationManager_Receiver = ConcertSyncClient::TestInterface::CreateClientReplicationManager(Client_Receiver->ClientSessionMock, BridgeUsed_Receiver.Get());
+		
 		// Client Sender
 		Client_Sender = &ConnectClient();
-		BridgeMock_Sender = MakeShared<FConcertClientReplicationBridgeMock>();
-		ClientReplicationManager_Sender = ConcertSyncClient::TestInterface::CreateClientReplicationManager(Client_Sender->ClientSessionMock, BridgeMock_Sender.Get());
+		if (EnumHasAnyFlags(Flags, ESendReceiveTestFlags::UseRealReplicationBridge))
+		{
+			BridgeUsed_Sender = ConcertSyncClient::TestInterface::CreateClientReplicationBridge();
+		}
+		else
+		{
+			BridgeMock_Sender =  MakeShared<FConcertClientReplicationBridgeMock>();
+			BridgeUsed_Sender = BridgeMock_Sender;
+		}
+		ClientReplicationManager_Sender = ConcertSyncClient::TestInterface::CreateClientReplicationManager(Client_Sender->ClientSessionMock, BridgeUsed_Sender.Get());
 		
 		// 1.1 Sender offers to send all UTestReflectionObject properties
 		{
@@ -130,8 +147,10 @@ namespace UE::ConcertSyncTests::Replication
 		{
 			ClientReplicationManager_Sender.Reset();
 			BridgeMock_Sender.Reset();
+			BridgeUsed_Sender.Reset();
 			ClientReplicationManager_Receiver.Reset();
 			BridgeMock_Receiver.Reset();
+			BridgeUsed_Sender.Reset();
 			ServerReplicationManager.Reset();
 			ServerSession.Reset();
 		}
diff --git a/Engine/Plugins/Developer/Concert/ConcertSync/ConcertSyncTest/Source/ConcertSyncTest/Private/Replication/Util/SendReceiveTestBase.h b/Engine/Plugins/Developer/Concert/ConcertSync/ConcertSyncTest/Source/ConcertSyncTest/Private/Replication/Util/SendReceiveTestBase.h
--- a/Engine/Plugins/Developer/Concert/ConcertSync/ConcertSyncTest/Source/ConcertSyncTest/Private/Replication/Util/SendReceiveTestBase.h
+++ b/Engine/Plugins/Developer/Concert/ConcertSync/ConcertSyncTest/Source/ConcertSyncTest/Private/Replication/Util/SendReceiveTestBase.h
@@ -6,6 +6,7 @@
 #include "Replication/Messages/ObjectReplication.h"
 #include "Util/ClientServerCommunicationTest.h"
 
+class IConcertClientReplicationBridge;
 class IConcertClientReplicationManager;
 class UTestReflectionObject;
 
@@ -21,6 +22,17 @@ namespace UE::ConcertSyncServer::Replication
 namespace UE::ConcertSyncTests::Replication
 {
 	class FConcertClientReplicationBridgeMock;
+
+	enum class ESendReceiveTestFlags : uint8
+	{
+		None = 0,
+		/**
+		 * Instead of mocking the bridge use the real implementation. The object will only be discovered if it exists when the stream is registered.
+		 * See IConcertClientReplicationBridge::PushTrackedObjects.
+		 */
+		UseRealReplicationBridge = 1 << 1
+	};
+	ENUM_CLASS_FLAGS(ESendReceiveTestFlags);
 	
 	/** Creates a server, connects a sender and receiver client, and completes a handshake for them. */
 	class FSendReceiveTestBase : public FConcertClientServerCommunicationTest
@@ -51,17 +63,23 @@ namespace UE::ConcertSyncTests::Replication
 		TSharedPtr<ConcertSyncServer::Replication::IConcertServerReplicationManager> ServerReplicationManager;
 		
 		FClientInfo* Client_Receiver = nullptr;
+		/** This is null if ESendReceiveTestFlags::UseRealReplicationBridge was specified. */
 		TSharedPtr<FConcertClientReplicationBridgeMock> BridgeMock_Receiver;
+		/** This is always valid. It is the bridge being used by the receiver. */
+		TSharedPtr<IConcertClientReplicationBridge> BridgeUsed_Receiver;
 		TSharedPtr<IConcertClientReplicationManager> ClientReplicationManager_Receiver;
 
 		FClientInfo* Client_Sender  = nullptr;
+		/** This is null if ESendReceiveTestFlags::UseRealReplicationBridge was specified. */
 		TSharedPtr<FConcertClientReplicationBridgeMock> BridgeMock_Sender;
+		/** This is always valid. It is the bridge being used by the receiver. */
+		TSharedPtr<IConcertClientReplicationBridge> BridgeUsed_Sender;
 		TSharedPtr<IConcertClientReplicationManager> ClientReplicationManager_Sender;
 
 		virtual ConcertSyncClient::Replication::FJoinReplicatedSessionArgs CreateSenderArgs() = 0;
 		virtual ConcertSyncClient::Replication::FJoinReplicatedSessionArgs CreateReceiverArgs() = 0;
 
-		virtual void SetUpClientAndServer();
+		virtual void SetUpClientAndServer(ESendReceiveTestFlags Flags = ESendReceiveTestFlags::None);
 		virtual void SimulateSenderToReceiver(
 			TFunctionRef<FReceiveReplicationEventSignature> OnServerReceive = [](auto&, auto&){},
 			TFunctionRef<FReceiveReplicationEventSignature> OnReceiverClientReceive = [](auto&, auto&){}
@@ -69,9 +87,9 @@ namespace UE::ConcertSyncTests::Replication
 		
 		void TickClient(FClientInfo* Client);
 		void TickServer();
-		
-	private:
-		
+
+		//~ Begin FConcertClientServerCommunicationTest Interface
 		virtual void CleanUpTest(FAutomationTestBase* AutomationTestBase) override;
+		//~ End FConcertClientServerCommunicationTest Interface
 	};
 }
diff --git a/Engine/Plugins/Developer/Concert/ConcertSync/ConcertSyncTest/Source/ConcertSyncTest/Private/Replication/Util/SendReceiveWorldTestBase.cpp b/Engine/Plugins/Developer/Concert/ConcertSync/ConcertSyncTest/Source/ConcertSyncTest/Private/Replication/Util/SendReceiveWorldTestBase.cpp
new file mode 100644
--- /dev/null
+++ b/Engine/Plugins/Developer/Concert/ConcertSync/ConcertSyncTest/Source/ConcertSyncTest/Private/Replication/Util/SendReceiveWorldTestBase.cpp
@@ -0,0 +1,54 @@
+
+#include "Replication/Util/SendReceiveWorldTestBase.h"
+
+namespace UE::ConcertSyncTests::Replication
+{
+	void FSendReceiveWorldTestBase::SetupWorld()
+	{
+		WorldManager = MakeUnique<FPreviewScene>(
+			FPreviewScene::ConstructionValues()
+			.SetEditor(true)
+			);
+	}
+
+	UWorld* FSendReceiveWorldTestBase::GetWorld() const
+	{
+		return ensureMsgf(WorldManager.IsValid(), TEXT("Your test did not call InitWorld!"))
+			? WorldManager->GetWorld()
+			: nullptr;
+	}
+
+	void FSendReceiveWorldTestBase::SimulateSenderToReceiverForWorld(TFunctionRef<void(UWorld& World)> SetTestValues, TFunctionRef<void(UWorld& World)> SetGarbageValues, TOptional<TFunctionRef<void(UWorld& World)>> TestValues)
+	{
+		UWorld* World = GetWorld();
+		if (!ensureMsgf(World, TEXT("Your test did not call InitWorld!")))
+		{
+			return;
+		}
+		
+		// TestObject is the same UObject on both clients.
+		// Hence we must override test values with SetTestValues and SetDifferentValues.
+		// 1 Sender > Server
+		SetTestValues(*World);
+		TickClient(Client_Sender);
+		
+		// 2 Forward from server to receiver
+		TickServer();
+		
+		// 3 Receive from server
+		SetGarbageValues(*World);
+		TickClient(Client_Receiver);
+		
+		if (TestValues)
+		{
+			TestValues.GetValue()(*World);
+		}
+		// No call to Super because we're completely overriding the behavior.
+	}
+
+	void FSendReceiveWorldTestBase::CleanUpTest(FAutomationTestBase* AutomationTestBase)
+	{
+		WorldManager.Reset();
+		Super::CleanUpTest(AutomationTestBase);
+	}
+}
\ No newline at end of file
diff --git a/Engine/Plugins/Developer/Concert/ConcertSync/ConcertSyncTest/Source/ConcertSyncTest/Private/Replication/Util/SendReceiveWorldTestBase.h b/Engine/Plugins/Developer/Concert/ConcertSync/ConcertSyncTest/Source/ConcertSyncTest/Private/Replication/Util/SendReceiveWorldTestBase.h
new file mode 100644
--- /dev/null
+++ b/Engine/Plugins/Developer/Concert/ConcertSync/ConcertSyncTest/Source/ConcertSyncTest/Private/Replication/Util/SendReceiveWorldTestBase.h
@@ -0,0 +1,65 @@
+// Copyright Epic Games, Inc. All Rights Reserved.
+
+#pragma once
+
+#include "PreviewScene.h"
+#include "SendReceiveGenericStreamTestBase.h"
+#include "Templates/Function.h"
+#include "Templates/UniquePtr.h"
+
+class UWorld;
+
+namespace UE::ConcertSyncTests::Replication
+{
+	class FConcertClientReplicationBridgeMock;
+	
+	/** For tests that want to test replication in a UWorld. */
+	class FSendReceiveWorldTestBase : public FSendReceiveGenericTestBase
+	{
+		using Super = FSendReceiveGenericTestBase;
+	public:
+
+		FSendReceiveWorldTestBase(const FString& InName, const bool bInComplexTask)
+			: Super(InName, bInComplexTask)
+		{}
+
+		/** Called by test when it is ready to initialize the world. */
+		virtual void SetupWorld();
+
+		/** @return The test's world. */
+		UWorld* GetWorld() const;
+
+		/**
+		 * Simulates sending values from one world to the other.
+		 * 
+		 * This actually operates on the same UWorld instance so that FSoftObjectPaths continue working.
+		 * This is how to use this function:
+		 * 1. Set up all your replicated objects
+		 * 2. SetTestValues sets values on whatever world objects you want
+		 * 3. This function reads said data and replicates it to the server
+		 * 4. The server forwards it to the receiving client
+		 * 5. SetGarbageValues sets values that are different from the test values
+		 * 6. This function applies the replicated data thus overriding the garbage values (that is assuming replication worked as expected!)
+		 * 7. This function executes TestValues if set so you can test whether the data was indeed replicated correctly.
+		 * 
+		 * @param SetTestValues Modifies values on replicated properties before it it sent from the sender client
+		 * @param SetGarbageValues Modifies values on replicated properties before the replicated data is applied on the receiving client
+		 * @param TestValues Optionally tests that the replicated data was applied correctly.
+		 * 
+		 */
+		void SimulateSenderToReceiverForWorld(
+			TFunctionRef<void(UWorld& World)> SetTestValues,
+			TFunctionRef<void(UWorld& World)> SetGarbageValues,
+			TOptional<TFunctionRef<void(UWorld& World)>> TestValues = {}
+			);
+
+	private:
+
+		/** Set when the test is started. Reset by CleanUpTest. */
+		TUniquePtr<FPreviewScene> WorldManager;
+		
+		//~ Begin FConcertClientServerCommunicationTest Interface
+		virtual void CleanUpTest(FAutomationTestBase* AutomationTestBase) override;
+		//~ End FConcertClientServerCommunicationTest Interface
+	};
+}
