Fix shadowing parameter warning for clang 18.1.  It seems to be just a bug with clang and we can work
around the issue by renaming the variable.

diff --git a/Engine/Plugins/Developer/Concert/ConcertSync/ConcertSyncTest/Source/ConcertSyncTest/Private/Replication/StreamRequestTests_Change.cpp b/Engine/Plugins/Developer/Concert/ConcertSync/ConcertSyncTest/Source/ConcertSyncTest/Private/Replication/StreamRequestTests_Change.cpp
--- a/Engine/Plugins/Developer/Concert/ConcertSync/ConcertSyncTest/Source/ConcertSyncTest/Private/Replication/StreamRequestTests_Change.cpp
+++ b/Engine/Plugins/Developer/Concert/ConcertSync/ConcertSyncTest/Source/ConcertSyncTest/Private/Replication/StreamRequestTests_Change.cpp
@@ -163,14 +163,14 @@ namespace UE::ConcertSyncTests::Replication::Stream
 		// Server logs a warning when rejecting - avoid the test being marked with a warning.
 		AddExpectedError(TEXT("Rejecting ChangeStream"));
 		ClientReplicationManager_Receiver->ChangeStream(AddFloatToStreamRequest)
-			.Next([this, TestObject, SenderStreamID = SenderStreamID, &TestObjectInReceiverStreamId, &bReceivedResponseAddingFloat](FConcertReplication_ChangeStream_Response&& Response)
+			.Next([this, TestObject, TheSenderStreamID = SenderStreamID, &TestObjectInReceiverStreamId, &bReceivedResponseAddingFloat](FConcertReplication_ChangeStream_Response&& Response)
 			{
 				bReceivedResponseAddingFloat = true;
 				TestTrue(TEXT("ErrorCode == Handled"), Response.ErrorCode == EReplicationResponseErrorCode::Handled);
 				TestEqual(TEXT("Append Float > 1 conflict"), Response.AuthorityConflicts.Num(), 1);
 				if (const FConcertReplicatedObjectId* ConflictingObject = Response.AuthorityConflicts.Find(TestObjectInReceiverStreamId))
 				{
-					TestEqual(TEXT("Append float > Conflict > Sender Stream correct"), ConflictingObject->StreamId, SenderStreamID);
+					TestEqual(TEXT("Append float > Conflict > Sender Stream correct"), ConflictingObject->StreamId, TheSenderStreamID);
 					TestEqual(TEXT("Append float > Conflict > Object correct"), ConflictingObject->Object, FSoftObjectPath(TestObject));
 					TestEqual(TEXT("Append float > Conflict > Endpoint ID correct"), ConflictingObject->SenderEndpointId, Client_Sender->ClientSessionMock->GetSessionClientEndpointId());
 				}
@@ -294,4 +294,4 @@ namespace UE::ConcertSyncTests::Replication::Stream
     	
     	return true;
     }
-}
\ No newline at end of file
+}
