Add "sync control" to Concert replication: clients will only replicate their data if there are clients there to receive it.
- Sending clients must register streams and take authority as normal
- New: The server sends FConcertReplication_ChangeSyncControl to authoritative clients when it is time to start & stop replicating

diff --git a/Engine/Plugins/Developer/Concert/ConcertSync/ConcertSyncClient/Source/ConcertSyncClient/Private/Replication/IConcertClientReplicationManager.cpp b/Engine/Plugins/Developer/Concert/ConcertSync/ConcertSyncClient/Source/ConcertSyncClient/Private/Replication/IConcertClientReplicationManager.cpp
--- a/Engine/Plugins/Developer/Concert/ConcertSync/ConcertSyncClient/Source/ConcertSyncClient/Private/Replication/IConcertClientReplicationManager.cpp
+++ b/Engine/Plugins/Developer/Concert/ConcertSync/ConcertSyncClient/Source/ConcertSyncClient/Private/Replication/IConcertClientReplicationManager.cpp
@@ -114,3 +114,21 @@ TMap<FSoftObjectPath, TSet<FGuid>> IConcertClientReplicationManager::GetClientOw
 	});
 	return Result;
 }
+
+TSet<FConcertObjectInStreamID> IConcertClientReplicationManager::GetSyncControlledObjects() const
+{
+	TSet<FConcertObjectInStreamID> Result;
+	const uint32 NumObjects = NumSyncControlledObjects();
+	if (NumObjects == 0)
+	{
+		return Result;
+	}
+	
+	Result.Reserve(NumObjects);
+	ForEachSyncControlledObject([&](const FConcertObjectInStreamID& Object)
+	{
+		Result.Add(Object);
+		return EBreakBehavior::Continue;
+	});
+	return Result;
+}
diff --git a/Engine/Plugins/Developer/Concert/ConcertSync/ConcertSyncClient/Source/ConcertSyncClient/Private/Replication/Manager/ReplicationManager.cpp b/Engine/Plugins/Developer/Concert/ConcertSync/ConcertSyncClient/Source/ConcertSyncClient/Private/Replication/Manager/ReplicationManager.cpp
--- a/Engine/Plugins/Developer/Concert/ConcertSync/ConcertSyncClient/Source/ConcertSyncClient/Private/Replication/Manager/ReplicationManager.cpp
+++ b/Engine/Plugins/Developer/Concert/ConcertSync/ConcertSyncClient/Source/ConcertSyncClient/Private/Replication/Manager/ReplicationManager.cpp
@@ -99,33 +99,69 @@ namespace UE::ConcertSyncClient::Replication
 			: TSet<FGuid>{}; 
 	}
 
+	IConcertClientReplicationManager::ESyncControlEnumerationResult FReplicationManager::ForEachSyncControlledObject(TFunctionRef<EBreakBehavior(const FConcertObjectInStreamID& Object)> Callback) const
+	{
+		// Check() to avoid returning some dummy static variable
+		ensureMsgf(CurrentState, TEXT("StartAcceptingJoinRequests should have been called at this point."));
+		return CurrentState
+			? CurrentState->ForEachSyncControlledObject(Callback)
+			: ESyncControlEnumerationResult::NoneAvailable;
+	}
+
+	uint32 FReplicationManager::NumSyncControlledObjects() const
+	{
+		return ensureMsgf(CurrentState, TEXT("StartAcceptingJoinRequests should have been called at this point."))
+			? CurrentState->NumSyncControlledObjects()
+			: 0;
+	}
+
+	bool FReplicationManager::HasSyncControl(const FConcertObjectInStreamID& Object) const
+	{
+		return ensureMsgf(CurrentState, TEXT("StartAcceptingJoinRequests should have been called at this point."))
+			&& CurrentState->HasSyncControl(Object);
+	}
+
 	IConcertClientReplicationManager::FOnPreStreamsChanged& FReplicationManager::OnPreStreamsChanged()
 	{
 		// Check() to avoid returning some dummy static variable
-		checkf(CurrentState, TEXT("StartAcceptingJoinRequests should have been called at this point."));
+		ensureMsgf(CurrentState, TEXT("StartAcceptingJoinRequests should have been called at this point."));
 		return CurrentState->OnPreStreamsChanged();
 	}
 
 	IConcertClientReplicationManager::FOnPostStreamsChanged& FReplicationManager::OnPostStreamsChanged()
 	{
 		// Check() to avoid returning some dummy static variable
-		checkf(CurrentState, TEXT("StartAcceptingJoinRequests should have been called at this point."));
+		ensureMsgf(CurrentState, TEXT("StartAcceptingJoinRequests should have been called at this point."));
 		return CurrentState->OnPostStreamsChanged();
 	}
 
 	IConcertClientReplicationManager::FOnPreAuthorityChanged& FReplicationManager::OnPreAuthorityChanged()
 	{
 		// Check() to avoid returning some dummy static variable
-		checkf(CurrentState, TEXT("StartAcceptingJoinRequests should have been called at this point."));
+		ensureMsgf(CurrentState, TEXT("StartAcceptingJoinRequests should have been called at this point."));
 		return CurrentState->OnPreAuthorityChanged();
 	}
 
 	IConcertClientReplicationManager::FOnPostAuthorityChanged& FReplicationManager::OnPostAuthorityChanged()
 	{
 		// Check() to avoid returning some dummy static variable
-		checkf(CurrentState, TEXT("StartAcceptingJoinRequests should have been called at this point."));
+		ensureMsgf(CurrentState, TEXT("StartAcceptingJoinRequests should have been called at this point."));
 		return CurrentState->OnPostAuthorityChanged();
 	}
+	
+	IConcertClientReplicationManager::FSyncControlChanged& FReplicationManager::OnPreSyncControlChanged()
+	{
+		// Check() to avoid returning some dummy static variable
+		ensureMsgf(CurrentState, TEXT("StartAcceptingJoinRequests should have been called at this point."));
+		return CurrentState->OnPreSyncControlChanged();
+	}
+
+	IConcertClientReplicationManager::FSyncControlChanged& FReplicationManager::OnPostSyncControlChanged()
+	{
+		// Check() to avoid returning some dummy static variable
+		ensureMsgf(CurrentState, TEXT("StartAcceptingJoinRequests should have been called at this point."));
+		return CurrentState->OnPostSyncControlChanged();
+	}
 
 	void FReplicationManager::OnChangeState(TSharedRef<FReplicationManagerState> NewState)
 	{
diff --git a/Engine/Plugins/Developer/Concert/ConcertSync/ConcertSyncClient/Source/ConcertSyncClient/Private/Replication/Manager/ReplicationManager.h b/Engine/Plugins/Developer/Concert/ConcertSync/ConcertSyncClient/Source/ConcertSyncClient/Private/Replication/Manager/ReplicationManager.h
--- a/Engine/Plugins/Developer/Concert/ConcertSync/ConcertSyncClient/Source/ConcertSyncClient/Private/Replication/Manager/ReplicationManager.h
+++ b/Engine/Plugins/Developer/Concert/ConcertSync/ConcertSyncClient/Source/ConcertSyncClient/Private/Replication/Manager/ReplicationManager.h
@@ -35,10 +35,15 @@ namespace UE::ConcertSyncClient::Replication
 		virtual TFuture<FConcertReplication_ChangeStream_Response> ChangeStream(FConcertReplication_ChangeStream_Request Args) override;
 		virtual EAuthorityEnumerationResult ForEachClientOwnedObject(TFunctionRef<EBreakBehavior(const FSoftObjectPath& Object, TSet<FGuid>&& OwningStreams)>) const override;
 		virtual TSet<FGuid> GetClientOwnedStreamsForObject(const FSoftObjectPath& ObjectPath) const override;
+		virtual ESyncControlEnumerationResult ForEachSyncControlledObject(TFunctionRef<EBreakBehavior(const FConcertObjectInStreamID& Object)> Callback) const override;
+		virtual uint32 NumSyncControlledObjects() const override;
+		virtual bool HasSyncControl(const FConcertObjectInStreamID& Object) const override;
 		virtual FOnPreStreamsChanged& OnPreStreamsChanged() override;
 		virtual FOnPostStreamsChanged& OnPostStreamsChanged() override;
 		virtual FOnPreAuthorityChanged& OnPreAuthorityChanged() override;
 		virtual FOnPostAuthorityChanged& OnPostAuthorityChanged() override;
+		virtual FSyncControlChanged& OnPreSyncControlChanged() override;
+		virtual FSyncControlChanged& OnPostSyncControlChanged() override;
 		//~ End IConcertClientReplicationManager Interface
 
 	private:
diff --git a/Engine/Plugins/Developer/Concert/ConcertSync/ConcertSyncClient/Source/ConcertSyncClient/Private/Replication/Manager/ReplicationManagerState.h b/Engine/Plugins/Developer/Concert/ConcertSync/ConcertSyncClient/Source/ConcertSyncClient/Private/Replication/Manager/ReplicationManagerState.h
--- a/Engine/Plugins/Developer/Concert/ConcertSync/ConcertSyncClient/Source/ConcertSyncClient/Private/Replication/Manager/ReplicationManagerState.h
+++ b/Engine/Plugins/Developer/Concert/ConcertSync/ConcertSyncClient/Source/ConcertSyncClient/Private/Replication/Manager/ReplicationManagerState.h
@@ -28,10 +28,15 @@ namespace UE::ConcertSyncClient::Replication
 		virtual TFuture<FConcertReplication_ChangeStream_Response> ChangeStream(FConcertReplication_ChangeStream_Request Args) override { return MakeFulfilledPromise<FConcertReplication_ChangeStream_Response>().GetFuture(); }
 		virtual EAuthorityEnumerationResult ForEachClientOwnedObject(TFunctionRef<EBreakBehavior(const FSoftObjectPath& Object, TSet<FGuid>&& OwningStreams)> Callback) const override { return EAuthorityEnumerationResult::NoAuthorityAvailable; }
 		virtual TSet<FGuid> GetClientOwnedStreamsForObject(const FSoftObjectPath& ObjectPath) const override { return {}; }
+		virtual ESyncControlEnumerationResult ForEachSyncControlledObject(TFunctionRef<EBreakBehavior(const FConcertObjectInStreamID& Object)> Callback) const override { return ESyncControlEnumerationResult::NoneAvailable; }
+		virtual uint32 NumSyncControlledObjects() const override { return 0; }
+		virtual bool HasSyncControl(const FConcertObjectInStreamID& Object) const override { return false; }
 		virtual FOnPreStreamsChanged& OnPreStreamsChanged() override { return OnPreStreamsChangedDelegate; } 
 		virtual FOnPostStreamsChanged& OnPostStreamsChanged() override { return OnPostStreamsChangedDelegate; }
 		virtual FOnPreAuthorityChanged& OnPreAuthorityChanged() override { return OnPreAuthorityChangedDelegate; }
 		virtual FOnPostAuthorityChanged& OnPostAuthorityChanged() override { return OnPostAuthorityChangedDelegate; }
+		virtual FSyncControlChanged& OnPreSyncControlChanged() override { return OnPreSyncControlChangedDelegate; }
+		virtual FSyncControlChanged& OnPostSyncControlChanged() override { return OnPostSyncControlChangedDelegate; }
 		//~ End IConcertClientReplicationManager Interface
 
 	protected:
@@ -40,6 +45,8 @@ namespace UE::ConcertSyncClient::Replication
 		FOnPostStreamsChanged OnPostStreamsChangedDelegate;
 		FOnPreAuthorityChanged OnPreAuthorityChangedDelegate;
 		FOnPostAuthorityChanged OnPostAuthorityChangedDelegate;
+		FSyncControlChanged OnPreSyncControlChangedDelegate;
+		FSyncControlChanged OnPostSyncControlChangedDelegate;
 
 		/**
 		 * Subclasses can change the state with this function.
diff --git a/Engine/Plugins/Developer/Concert/ConcertSync/ConcertSyncClient/Source/ConcertSyncClient/Private/Replication/Manager/ReplicationManagerState_Connected.cpp b/Engine/Plugins/Developer/Concert/ConcertSync/ConcertSyncClient/Source/ConcertSyncClient/Private/Replication/Manager/ReplicationManagerState_Connected.cpp
--- a/Engine/Plugins/Developer/Concert/ConcertSync/ConcertSyncClient/Source/ConcertSyncClient/Private/Replication/Manager/ReplicationManagerState_Connected.cpp
+++ b/Engine/Plugins/Developer/Concert/ConcertSync/ConcertSyncClient/Source/ConcertSyncClient/Private/Replication/Manager/ReplicationManagerState_Connected.cpp
@@ -69,20 +69,23 @@ namespace UE::ConcertSyncClient::Replication
 	}
 	
 	FReplicationManagerState_Connected::FReplicationManagerState_Connected(
-		TSharedRef<IConcertClientSession> LiveSession,
+		TSharedRef<IConcertClientSession> InLiveSession,
 		IConcertClientReplicationBridge& ReplicationBridge,
-		TArray<FConcertReplicationStream> StreamDescriptions,
-		FReplicationManager& Owner
+		FReplicationManager& Owner,
+		TArray<FConcertReplicationStream> InitialStreams,
+		const FConcertReplication_ChangeSyncControl& InitialSyncControl
 		)
 		: FReplicationManagerState(Owner)
-		, LiveSession(LiveSession)
+		, LiveSession(InLiveSession)
 		, ReplicationBridge(ReplicationBridge)
-		, RegisteredStreams(MoveTemp(StreamDescriptions))
+		, RegisteredStreams(MoveTemp(InitialStreams))
 		// TODO DP: Use config to determine which replication format to use
 		, ReplicationFormat(MakeUnique<ConcertSyncCore::FFullObjectFormat>())
+		, SyncControl(*LiveSession)
 		, ReplicationDataSource(
 			ReplicationBridge,
 			*ReplicationFormat,
+			SyncControl,
 			FClientReplicationDataCollector::FGetClientStreams::CreateLambda([this]()
 			{
 				return &RegisteredStreams;
@@ -91,13 +94,17 @@ namespace UE::ConcertSyncClient::Replication
 			)
 		, Sender(
 			ConcertSyncCore::FGetObjectFrequencySettings::CreateRaw(this, &FReplicationManagerState_Connected::GetObjectFrequencySettings),
-			LiveSession->GetSessionServerEndpointId(), *LiveSession, ReplicationDataSource
+			InLiveSession->GetSessionServerEndpointId(), *LiveSession, ReplicationDataSource
 			)
 		, ReceivedDataCache(MakeShared<ConcertSyncCore::FObjectReplicationCache>(*ReplicationFormat))
-		, Receiver(*LiveSession, *ReceivedDataCache)
+		, Receiver(*InLiveSession, *ReceivedDataCache)
 		, ReceivedReplicationQueuer(FClientReplicationDataQueuer::Make(ReplicationBridge, *ReceivedDataCache))
 		, ReplicationApplier(ReplicationBridge, *ReplicationFormat, *ReceivedReplicationQueuer)
-	{}
+	{
+		SyncControl.ProcessSyncControlChange(InitialSyncControl);
+		SyncControl.OnPreSyncControlChanged().AddLambda([this](){ OnPreSyncControlChangedDelegate.Broadcast(); });
+		SyncControl.OnPostSyncControlChanged().AddLambda([this](){ OnPostSyncControlChangedDelegate.Broadcast(); });
+	}
 
 	FReplicationManagerState_Connected::~FReplicationManagerState_Connected()
 	{
@@ -145,6 +152,8 @@ namespace UE::ConcertSyncClient::Replication
 		// Stop replicating removed objects right now: the server will remove authority after processing this request.
 		// At that point, it will log errors for receiving replication data from a client without authority.
 		HandleReleasingReplicatedObjects(Args);
+		// We don't need worry about updating sync control until it is processed below - the local client will not attempt to replicate the object
+		// because we just locally updated the authority cache.
 
 		Private::LogNetworkMessage(CVarLogAuthorityRequestsAndResponsesOnClient, Args);
 		return LiveSession->SendCustomRequest<FConcertReplication_ChangeAuthority_Request, FConcertReplication_ChangeAuthority_Response>(Args, LiveSession->GetSessionServerEndpointId())
@@ -155,6 +164,7 @@ namespace UE::ConcertSyncClient::Replication
 				if (const TSharedPtr<FReplicationManagerState_Connected> ThisPin = WeakThis.Pin()
 					; ThisPin && Response.ErrorCode == EReplicationResponseErrorCode::Handled)
 				{
+					ThisPin->SyncControl.ProcessAuthorityChange(Args, Response);
 					ThisPin->UpdateReplicatedObjectsAfterAuthorityChange(MoveTemp(Args), Response);
 				}
 				else if (ThisPin && Response.ErrorCode == EReplicationResponseErrorCode::Timeout)
@@ -197,6 +207,8 @@ namespace UE::ConcertSyncClient::Replication
 		// Stop replicating removed objects right now: the server will remove authority after processing this request.
 		// At that point, it will log errors for receiving replication data from a client without authority.
 		HandleRemovingReplicatedObjects(Args);
+		// We don't need worry about updating sync control until it is processed below - the local client will not attempt to replicate the object
+		// because we just locally updated the replication cache.
 		
 		Private::LogNetworkMessage(CVarLogStreamRequestsAndResponsesOnClient, Args);
 		return LiveSession->SendCustomRequest<FConcertReplication_ChangeStream_Request, FConcertReplication_ChangeStream_Response>(Args, LiveSession->GetSessionServerEndpointId())
@@ -207,6 +219,7 @@ namespace UE::ConcertSyncClient::Replication
 				const TSharedPtr<FReplicationManagerState_Connected> ThisPin = WeakThis.Pin();
 				if (ThisPin && Response.IsSuccess())
 				{
+					ThisPin->SyncControl.ProcessStreamChange(Args);
 					ThisPin->UpdateReplicatedObjectsAfterStreamChange(Args, Response);
 				}
 				else if (ThisPin && Response.ErrorCode == EReplicationResponseErrorCode::Timeout)
@@ -246,6 +259,13 @@ namespace UE::ConcertSyncClient::Replication
 		return Result;
 	}
 
+	IConcertClientReplicationManager::ESyncControlEnumerationResult FReplicationManagerState_Connected::ForEachSyncControlledObject(TFunctionRef<EBreakBehavior(const FConcertObjectInStreamID& Object)> Callback) const
+	{
+		return SyncControl.EnumerateAllowedObjects(Callback)
+			? ESyncControlEnumerationResult::Iterated
+			: ESyncControlEnumerationResult::NoneAvailable;
+	}
+
 	void FReplicationManagerState_Connected::OnEnterState()
 	{
 		LiveSession->OnTick().AddSP(this, &FReplicationManagerState_Connected::Tick);
diff --git a/Engine/Plugins/Developer/Concert/ConcertSync/ConcertSyncClient/Source/ConcertSyncClient/Private/Replication/Manager/ReplicationManagerState_Connected.h b/Engine/Plugins/Developer/Concert/ConcertSync/ConcertSyncClient/Source/ConcertSyncClient/Private/Replication/Manager/ReplicationManagerState_Connected.h
--- a/Engine/Plugins/Developer/Concert/ConcertSync/ConcertSyncClient/Source/ConcertSyncClient/Private/Replication/Manager/ReplicationManagerState_Connected.h
+++ b/Engine/Plugins/Developer/Concert/ConcertSync/ConcertSyncClient/Source/ConcertSyncClient/Private/Replication/Manager/ReplicationManagerState_Connected.h
@@ -4,6 +4,7 @@
 
 #include "ReplicationManagerState.h"
 #include "Replication/Formats/IObjectReplicationFormat.h"
+#include "Replication/Misc/LocalSyncControl.h"
 #include "Replication/Processing/ClientReplicationDataCollector.h"
 #include "Replication/Processing/ClientReplicationDataQueuer.h"
 #include "Replication/Processing/ObjectReplicationApplierProcessor.h"
@@ -29,10 +30,11 @@ namespace UE::ConcertSyncClient::Replication
 	public:
 
 		FReplicationManagerState_Connected(
-			TSharedRef<IConcertClientSession> LiveSession,
+			TSharedRef<IConcertClientSession> InLiveSession,
 			IConcertClientReplicationBridge& ReplicationBridge UE_LIFETIMEBOUND,
-			TArray<FConcertReplicationStream> StreamDescriptions,
-			FReplicationManager& Owner UE_LIFETIMEBOUND
+			FReplicationManager& Owner UE_LIFETIMEBOUND,
+			TArray<FConcertReplicationStream> InitialStreams,
+			const FConcertReplication_ChangeSyncControl& InitialSyncControl
 			);
 		virtual ~FReplicationManagerState_Connected() override;
 
@@ -47,6 +49,9 @@ namespace UE::ConcertSyncClient::Replication
 		virtual TFuture<FConcertReplication_ChangeStream_Response> ChangeStream(FConcertReplication_ChangeStream_Request Args) override;
 		virtual EAuthorityEnumerationResult ForEachClientOwnedObject(TFunctionRef<EBreakBehavior(const FSoftObjectPath& Object, TSet<FGuid>&& OwningStreams)> Callback) const override;
 		virtual TSet<FGuid> GetClientOwnedStreamsForObject(const FSoftObjectPath& ObjectPath) const override;
+		virtual ESyncControlEnumerationResult ForEachSyncControlledObject(TFunctionRef<EBreakBehavior(const FConcertObjectInStreamID& Object)> Callback) const override;
+		virtual uint32 NumSyncControlledObjects() const override { return SyncControl.Num(); }
+		virtual bool HasSyncControl(const FConcertObjectInStreamID& Object) const override { return SyncControl.IsObjectAllowed(Object); }
 		//~ End IConcertClientReplicationManager Interface
 
 	private:
@@ -62,6 +67,8 @@ namespace UE::ConcertSyncClient::Replication
 		const TUniquePtr<ConcertSyncCore::IObjectReplicationFormat> ReplicationFormat;
 
 		// Sending
+		/** Decides whether an object should be replicated. */
+		FLocalSyncControl SyncControl;
 		/** Used as source of replication data. */
 		FClientReplicationDataCollector ReplicationDataSource;
 		
diff --git a/Engine/Plugins/Developer/Concert/ConcertSync/ConcertSyncClient/Source/ConcertSyncClient/Private/Replication/Manager/ReplicationManagerState_Handshaking.cpp b/Engine/Plugins/Developer/Concert/ConcertSync/ConcertSyncClient/Source/ConcertSyncClient/Private/Replication/Manager/ReplicationManagerState_Handshaking.cpp
--- a/Engine/Plugins/Developer/Concert/ConcertSync/ConcertSyncClient/Source/ConcertSyncClient/Private/Replication/Manager/ReplicationManagerState_Handshaking.cpp
+++ b/Engine/Plugins/Developer/Concert/ConcertSync/ConcertSyncClient/Source/ConcertSyncClient/Private/Replication/Manager/ReplicationManagerState_Handshaking.cpp
@@ -62,7 +62,7 @@ namespace UE::ConcertSyncClient::Replication
 				if (Response.JoinErrorCode == EJoinReplicationErrorCode::Success)
 				{
 					This->ChangeState(
-						MakeShared<FReplicationManagerState_Connected>(This->LiveSession, This->ReplicationBridge, MoveTemp(This->RequestArgs.Streams), This->GetOwner())
+						MakeShared<FReplicationManagerState_Connected>(This->LiveSession, This->ReplicationBridge, This->GetOwner(), MoveTemp(This->RequestArgs.Streams), Response.SyncControl)
 						);
 				}
 				else
diff --git a/Engine/Plugins/Developer/Concert/ConcertSync/ConcertSyncClient/Source/ConcertSyncClient/Private/Replication/Misc/LocalSyncControl.h b/Engine/Plugins/Developer/Concert/ConcertSync/ConcertSyncClient/Source/ConcertSyncClient/Private/Replication/Misc/LocalSyncControl.h
new file mode 100644
--- /dev/null
+++ b/Engine/Plugins/Developer/Concert/ConcertSync/ConcertSyncClient/Source/ConcertSyncClient/Private/Replication/Misc/LocalSyncControl.h
@@ -0,0 +1,85 @@
+﻿// Copyright Epic Games, Inc. All Rights Reserved.
+
+#pragma once
+
+#include "IConcertSession.h"
+#include "Replication/Messages/SyncControl.h"
+
+#include "HAL/Platform.h"
+#include "Replication/SyncControlState.h"
+#include "Templates/UnrealTemplate.h"
+
+namespace UE::ConcertSyncClient::Replication
+{
+	/**
+	 * Receives network messages that explicitly and implicitly change sync control.
+	 * Those network messages are fed into FSyncControlState, which parses the messages.
+	 * 
+	 * Sync control is the set of objects this client is allowed to replicate.
+	 */
+	class FLocalSyncControl
+		: public FNoncopyable
+		, ConcertSyncCore::Replication::FSyncControlState
+	{
+	public:
+
+		explicit FLocalSyncControl(IConcertSession& InSession UE_LIFETIMEBOUND)
+			: Session(InSession)
+		{
+			Session.RegisterCustomEventHandler<FConcertReplication_ChangeSyncControl>(this, &FLocalSyncControl::OnChangeSyncControl);
+		}
+
+		~FLocalSyncControl()
+		{
+			Session.UnregisterCustomEventHandler<FConcertReplication_ChangeSyncControl>(this);
+		}
+
+		using FSyncControlState::IsObjectAllowed;
+		using FSyncControlState::Num;
+		using FSyncControlState::EnumerateAllowedObjects;
+		using FSyncControlState::EnumerateChanges;
+		// No using for FSyncControlState::Aggregate because we want to limit usage to below ProcessXChange functions
+
+		/** Stores the given change in this data structure. */
+		void ProcessSyncControlChange(const FConcertReplication_ChangeSyncControl& Event)
+		{
+			OnPreSyncControlChangedDelegate.Broadcast();
+			AppendChanges(Event);
+			OnPostSyncControlChangedDelegate.Broadcast();
+		}
+
+		/** Applies implicit and explicit changes to the client's sync control resulting from a completed authority change request and response. */
+		void ProcessAuthorityChange(const FConcertReplication_ChangeAuthority_Request& Request, const FConcertReplication_ChangeAuthority_Response& Response)
+		{
+			OnPreSyncControlChangedDelegate.Broadcast();
+			AppendChanges(Request, Response);
+			OnPostSyncControlChangedDelegate.Broadcast();
+		}
+
+		/**
+		 * Applies implicit changes to the client's sync control resulting from losing authority from objects removed from the stream.
+		 * You must validate that the request has also been accepted by the server!
+		 */
+		void ProcessStreamChange(const FConcertReplication_ChangeStream_Request& Request)
+		{
+			OnPreSyncControlChangedDelegate.Broadcast();
+			AppendChanges(Request);
+			OnPostSyncControlChangedDelegate.Broadcast();
+		}
+		
+		DECLARE_MULTICAST_DELEGATE(FSyncControlChanged);
+		FSyncControlChanged& OnPreSyncControlChanged() { return OnPreSyncControlChangedDelegate; }
+		FSyncControlChanged& OnPostSyncControlChanged() { return OnPostSyncControlChangedDelegate; }
+
+	private:
+
+		/** The session to receive sync control changes on. */
+		IConcertSession& Session;
+		
+		FSyncControlChanged OnPreSyncControlChangedDelegate;
+		FSyncControlChanged OnPostSyncControlChangedDelegate;
+
+		void OnChangeSyncControl(const FConcertSessionContext&, const FConcertReplication_ChangeSyncControl& Event) { ProcessSyncControlChange(Event); }
+	};
+}
+
diff --git a/Engine/Plugins/Developer/Concert/ConcertSync/ConcertSyncClient/Source/ConcertSyncClient/Private/Replication/Processing/ClientReplicationDataCollector.cpp b/Engine/Plugins/Developer/Concert/ConcertSync/ConcertSyncClient/Source/ConcertSyncClient/Private/Replication/Processing/ClientReplicationDataCollector.cpp
--- a/Engine/Plugins/Developer/Concert/ConcertSync/ConcertSyncClient/Source/ConcertSyncClient/Private/Replication/Processing/ClientReplicationDataCollector.cpp
+++ b/Engine/Plugins/Developer/Concert/ConcertSync/ConcertSyncClient/Source/ConcertSyncClient/Private/Replication/Processing/ClientReplicationDataCollector.cpp
@@ -10,17 +10,20 @@
 
 #include "Algo/RemoveIf.h"
 #include "Misc/EBreakBehavior.h"
+#include "Replication/Misc/LocalSyncControl.h"
 
 namespace UE::ConcertSyncClient::Replication
 {
 	FClientReplicationDataCollector::FClientReplicationDataCollector(
 		IConcertClientReplicationBridge& InReplicationBridge,
 		ConcertSyncCore::IObjectReplicationFormat& InReplicationFormat,
+		const FLocalSyncControl& SyncControl,
 		FGetClientStreams InGetStreamsDelegate,
 		const FGuid& InClientId
 		)
 		: Bridge(InReplicationBridge)
 		, ReplicationFormat(InReplicationFormat)
+		, SyncControl(SyncControl)
 		, GetStreamsDelegate(MoveTemp(InGetStreamsDelegate))
 		, ClientId(InClientId)
 	{
@@ -161,6 +164,15 @@ namespace UE::ConcertSyncClient::Replication
 		{
 			for (const FObjectInfo& ObjectInfo : Pair.Value)
 			{
+				const FConcertReplicatedObjectId ObjectId{
+					FConcertObjectInStreamID{ ObjectInfo.StreamId, Pair.Key },
+					ClientId
+				};
+				if (!SyncControl.IsObjectAllowed(ObjectId))
+				{
+					continue;;
+				}
+				
 				// The bridge has the object cached. We do not cache the object ourselves!
 				// Funky Unreal flows can cause the object to be renamed out from under us and replaced by a different instance (just by using UObject::Rename()).
 				// The bridge is aware of these flows and FindObjectIfAvailable will catch them.
@@ -170,10 +182,7 @@ namespace UE::ConcertSyncClient::Replication
 				const FSoftObjectPath ObjectPath = Object;
 				if (Object && ensureMsgf(ObjectPath == Object, TEXT("Sanity check: the bridge gave us an object with a different path!")))
 				{
-					ProcessItemFunc(FConcertReplicatedObjectId {
-						FConcertObjectInStreamID{ ObjectInfo.StreamId, Object },
-						ClientId
-					});
+					ProcessItemFunc(ObjectId);
 				}
 			}
 		}
@@ -189,7 +198,7 @@ namespace UE::ConcertSyncClient::Replication
 		// ExtractReplicationDataForObject is supposed to be called in response to ForEachPendingObject... so either the call was invalid or ForEachPendingObject lied
 		if (!ensure(ObjectInfos))
 		{
-			return {};
+			return false;
 		}
 
 		// The properties to be replicated depend on the stream: search for it.
@@ -197,13 +206,13 @@ namespace UE::ConcertSyncClient::Replication
 		// Same logic as above: if the stream is not found, then either the call is invalid or ForEachPendingObject lied
 		if (!ensure(ObjectInfos->IsValidIndex(StreamIndex)))
 		{
-			return {};
+			return false;
 		}
 
 		UObject* Object = Bridge.FindObjectIfAvailable(ObjectToProcess.Object);
 		// Finally... ask the bridge to resolve the object for us. The bridge has the object cached and handles the object getting renamed, etc.
 		// This should resolve. If it does not: same logic as above - either this call is invalid or ForEachPendingObject lied.
-		if (!ensure(Object))
+		if (!ensure(Object) || !ensure(SyncControl.IsObjectAllowed(ObjectToProcess)))
 		{
 			return false;
 		}
diff --git a/Engine/Plugins/Developer/Concert/ConcertSync/ConcertSyncClient/Source/ConcertSyncClient/Private/Replication/Processing/ClientReplicationDataCollector.h b/Engine/Plugins/Developer/Concert/ConcertSync/ConcertSyncClient/Source/ConcertSyncClient/Private/Replication/Processing/ClientReplicationDataCollector.h
--- a/Engine/Plugins/Developer/Concert/ConcertSync/ConcertSyncClient/Source/ConcertSyncClient/Private/Replication/Processing/ClientReplicationDataCollector.h
+++ b/Engine/Plugins/Developer/Concert/ConcertSync/ConcertSyncClient/Source/ConcertSyncClient/Private/Replication/Processing/ClientReplicationDataCollector.h
@@ -26,6 +26,8 @@ namespace UE::ConcertSyncCore
 
 namespace UE::ConcertSyncClient::Replication
 {
+	class FLocalSyncControl;
+	
 	/**
 	 * Exposes UObject instances to an FObjectReplicationProcessor.
 	 * IConcertClientReplicationBridge tracks UObject lifetime, this class exposes them.
@@ -39,6 +41,7 @@ namespace UE::ConcertSyncClient::Replication
 		FClientReplicationDataCollector(
 			IConcertClientReplicationBridge& InReplicationBridge UE_LIFETIMEBOUND,
 			ConcertSyncCore::IObjectReplicationFormat& InReplicationFormat UE_LIFETIMEBOUND,
+			const FLocalSyncControl& SyncControl UE_LIFETIMEBOUND,
 			FGetClientStreams InGetStreamsDelegate,
 			const FGuid& InClientId
 			);
@@ -81,6 +84,9 @@ namespace UE::ConcertSyncClient::Replication
 		/** Used to create the replication data sent to the server. */
 		ConcertSyncCore::IObjectReplicationFormat& ReplicationFormat;
 
+		/** Tells us whether we're allowed to replicate an object in a stream. */
+		const FLocalSyncControl& SyncControl;
+
 		/** Gets the stream of the managed client. */
 		const FGetClientStreams GetStreamsDelegate;
 
diff --git a/Engine/Plugins/Developer/Concert/ConcertSync/ConcertSyncClient/Source/ConcertSyncClient/Public/Replication/IConcertClientReplicationManager.h b/Engine/Plugins/Developer/Concert/ConcertSync/ConcertSyncClient/Source/ConcertSyncClient/Public/Replication/IConcertClientReplicationManager.h
--- a/Engine/Plugins/Developer/Concert/ConcertSync/ConcertSyncClient/Source/ConcertSyncClient/Public/Replication/IConcertClientReplicationManager.h
+++ b/Engine/Plugins/Developer/Concert/ConcertSync/ConcertSyncClient/Source/ConcertSyncClient/Public/Replication/IConcertClientReplicationManager.h
@@ -11,6 +11,7 @@
 template<typename ResultType>
 class TFuture;
 
+struct FConcertReplication_ChangeSyncControl;
 struct FConcertReplicationStream;
 
 namespace UE::ConcertSyncClient::Replication
@@ -70,7 +71,7 @@ public:
 	 * @note This future can finish on any thread (e.g. when message endpoint times out); usually it finishes on the game thread.  
 	 */
 	virtual EStreamEnumerationResult ForEachRegisteredStream(TFunctionRef<EBreakBehavior(const FConcertReplicationStream& Stream)> Callback) const = 0;
-	/** @return Whether this manager is has any registered streams (basically whether ForEachRegisteredStream returns EStreamEnumerationResult::Iterated). */
+	/** @return Whether this manager has any registered streams (basically whether ForEachRegisteredStream returns EStreamEnumerationResult::Iterated). */
 	bool HasRegisteredStreams() const;
 	/** @return The streams registered with the server. */
 	TArray<FConcertReplicationStream> GetRegisteredStreams() const;
@@ -96,6 +97,16 @@ public:
 	/** @return All owned objects and associated owning streams. */
 	TMap<FSoftObjectPath, TSet<FGuid>> GetClientOwnedObjects() const;
 
+	enum class ESyncControlEnumerationResult { NoneAvailable, Iterated };
+	/** Iterates through all objects the server told the client to replicate, which happens when this client has taken authority and another client wants to receive the data. */
+	virtual ESyncControlEnumerationResult ForEachSyncControlledObject(TFunctionRef<EBreakBehavior(const FConcertObjectInStreamID& Object)> Callback) const = 0;
+	/** @return Number of items ForEachSyncControlledObject would iterate. */
+	virtual uint32 NumSyncControlledObjects() const = 0;
+	/** @return Whether this client has sync control for a specific object in a stream. */
+	virtual bool HasSyncControl(const FConcertObjectInStreamID& Object) const = 0;
+	/** Util for converting all sync controlled objects into a TSet. */
+	TSet<FConcertObjectInStreamID> GetSyncControlledObjects() const;
+
 	/**
 	 * Requests replication info about other clients, including the streams registered and which objects they have authority over (i.e. are sending).
 	 * @note This future can finish on any thread (e.g. when message endpoint times out); usually it finishes on the game thread.  
@@ -129,6 +140,12 @@ public:
 	DECLARE_MULTICAST_DELEGATE(FOnPostAuthorityChanged);
 	/** Called right after GetClientOwnedObjects has changed. */
 	virtual FOnPostAuthorityChanged& OnPostAuthorityChanged() = 0;
+
+	DECLARE_MULTICAST_DELEGATE(FSyncControlChanged);
+	/** Called just before sync control change is applied. */
+	virtual FSyncControlChanged& OnPreSyncControlChanged() = 0;
+	/** Called just after sync control change is applied. */
+	virtual FSyncControlChanged& OnPostSyncControlChanged() = 0;
 	
 	virtual ~IConcertClientReplicationManager() = default;
 };
diff --git a/Engine/Plugins/Developer/Concert/ConcertSync/ConcertSyncCore/Source/ConcertSyncCore/Public/Replication/Messages/ChangeAuthority.h b/Engine/Plugins/Developer/Concert/ConcertSync/ConcertSyncCore/Source/ConcertSyncCore/Public/Replication/Messages/ChangeAuthority.h
--- a/Engine/Plugins/Developer/Concert/ConcertSync/ConcertSyncCore/Source/ConcertSyncCore/Public/Replication/Messages/ChangeAuthority.h
+++ b/Engine/Plugins/Developer/Concert/ConcertSync/ConcertSyncCore/Source/ConcertSyncCore/Public/Replication/Messages/ChangeAuthority.h
@@ -3,6 +3,7 @@
 #pragma once
 
 #include "EReplicationResponseErrorCode.h"
+#include "SyncControl.h"
 #include "ChangeAuthority.generated.h"
 
 USTRUCT()
@@ -63,4 +64,14 @@ struct FConcertReplication_ChangeAuthority_Response
 	 */
 	UPROPERTY()
 	TMap<FSoftObjectPath, FConcertStreamArray> RejectedObjects;
+
+	/**
+	 * Contains the new sync control for objects you requested authority over; objects you released authority over are not contained because they can
+	 * be implicitly determined using FSyncControlState::Aggregate.
+	 * 
+	 * @see FSyncControlState::Aggregate.
+	 * @see FConcertReplication_ChangeSyncControl
+	 */
+	UPROPERTY()
+	FConcertReplication_ChangeSyncControl SyncControl;
 };
\ No newline at end of file
diff --git a/Engine/Plugins/Developer/Concert/ConcertSync/ConcertSyncCore/Source/ConcertSyncCore/Public/Replication/Messages/Handshake.h b/Engine/Plugins/Developer/Concert/ConcertSync/ConcertSyncCore/Source/ConcertSyncCore/Public/Replication/Messages/Handshake.h
--- a/Engine/Plugins/Developer/Concert/ConcertSync/ConcertSyncCore/Source/ConcertSyncCore/Public/Replication/Messages/Handshake.h
+++ b/Engine/Plugins/Developer/Concert/ConcertSync/ConcertSyncCore/Source/ConcertSyncCore/Public/Replication/Messages/Handshake.h
@@ -2,6 +2,7 @@
 
 #pragma once
 
+#include "SyncControl.h"
 #include "Replication/Data/ReplicationStream.h"
 #include "Handshake.generated.h"
 
@@ -84,6 +85,10 @@ struct FConcertReplication_Join_Response
 	/** More information about ErrorCode to help the user resolve the issue. */
 	UPROPERTY()
 	FString DetailedErrorMessage;
+
+	/** This is the initial sync control this client has */
+	UPROPERTY()
+	FConcertReplication_ChangeSyncControl SyncControl;
 };
 
 /** Sent by client to notify server that no more replication data should be sent. */
diff --git a/Engine/Plugins/Developer/Concert/ConcertSync/ConcertSyncCore/Source/ConcertSyncCore/Public/Replication/Messages/SyncControl.h b/Engine/Plugins/Developer/Concert/ConcertSync/ConcertSyncCore/Source/ConcertSyncCore/Public/Replication/Messages/SyncControl.h
new file mode 100644
--- /dev/null
+++ b/Engine/Plugins/Developer/Concert/ConcertSync/ConcertSyncCore/Source/ConcertSyncCore/Public/Replication/Messages/SyncControl.h
@@ -0,0 +1,44 @@
+﻿// Copyright Epic Games, Inc. All Rights Reserved.
+
+#pragma once
+
+#include "Replication/Data/ObjectIds.h"
+#include "SyncControl.generated.h"
+
+/**
+ * Notifies clients about the ability to start or stop replicating specific objects.
+ *
+ * Sync control serves to inform clients whether the data they wish to replicate has any receivers.
+ * Replicating data when no one is listening is unnecessary.
+ *
+ * Until a sync control event is received, clients should assume no interest in the data.
+ *
+ * Special Circumstances:
+ * 1. FConcertReplication_Join_Request: The server doesn't send any FConcertReplication_ChangeSyncControl event in response to the handshake.
+ *    Instead, the resulting sync control is in FConcertReplication_Join_Response::SyncControl.
+ * 2. Explicit authority change via FConcertReplication_ChangeAuthority_Request: No sync control event is sent.
+ *	  FConcertReplication_ChangeAuthority_Response::SyncControl contains the objects for which the client received instant sync control.
+ *    Objects losing sync control are implicitly determined by using FSyncControlState::Aggregate.
+ * 3. Implicit authority change via FConcertReplication_ChangeStream_Request: No sync control event is generated as its result is implicit from the request and response.
+ *	  Stream changes can only cause the client to lose authority, thus sync control, over objects.
+ *    Therefore, objects losing sync control can be enumerated using FSyncControlState::Aggregate.
+ *
+ * Example Flow:
+ * 1. Client 1 registers stream 's' with object 'x' (and associated properties) using FConcertReplication_ChangeStream_Request.
+ * 2. Client 1 takes authority over 'x' with 's' using FConcertReplication_ChangeAuthority_Request.
+ * 3. Client 2 joins (and wants to receive the data).
+ * 4. Server sends FConcertReplication_ChangeSyncControl to client 1 to begin replicating 'x' in 's'.
+ * 5. Client 2 leaves.
+ * 6. Server sends FConcertReplication_ChangeSyncControl to client 1 to cease replicating 'x' in 's'.
+ */
+USTRUCT()
+struct FConcertReplication_ChangeSyncControl
+{
+	GENERATED_BODY()
+
+	/** Toggles the objects (which the receiving client 1. has previously registered and 2. taken authority over) to replicate or not. */
+	UPROPERTY()
+	TMap<FConcertObjectInStreamID, bool> NewControlStates;
+
+	bool IsEmpty() const { return NewControlStates.IsEmpty(); }
+};
\ No newline at end of file
diff --git a/Engine/Plugins/Developer/Concert/ConcertSync/ConcertSyncCore/Source/ConcertSyncCore/Public/Replication/SyncControlState.h b/Engine/Plugins/Developer/Concert/ConcertSync/ConcertSyncCore/Source/ConcertSyncCore/Public/Replication/SyncControlState.h
new file mode 100644
--- /dev/null
+++ b/Engine/Plugins/Developer/Concert/ConcertSync/ConcertSyncCore/Source/ConcertSyncCore/Public/Replication/SyncControlState.h
@@ -0,0 +1,189 @@
+﻿// Copyright Epic Games, Inc. All Rights Reserved.
+
+#pragma once
+
+#include "Data/ObjectIds.h"
+#include "Messages/SyncControl.h"
+#include "Misc/EBreakBehavior.h"
+#include "Replication/Messages/ChangeAuthority.h"
+#include "Replication/Messages/ChangeStream.h"
+#include "Replication/Messages/SyncControl.h"
+
+#include "Templates/Function.h"
+
+#include <type_traits>
+
+namespace UE::ConcertSyncCore::Replication
+{
+	template<typename T>
+	concept CObjectInStreamInvocable = std::is_invocable_v<T, const FConcertObjectInStreamID&>;
+	
+	/**
+	 * Keeps track of the sync control (AllowedObjects), which is a set of objects a client is allowed to replicate.
+	 * It knows how to parse network messages that explicitly and implicitly change sync control and update AllowedObjects.
+	 * 
+	 * This class does not know where the network messages come from and is designed to be used as utility by lower level systems.
+	 */
+	class FSyncControlState
+	{
+	public:
+
+		FSyncControlState() = default;
+		
+		// Not explicit to make it useful to use
+		FSyncControlState(TSet<FConcertObjectInStreamID>&& AllowedObjects) 
+			: AllowedObjects(MoveTemp(AllowedObjects))
+		{}
+		FSyncControlState(const TSet<FConcertObjectInStreamID>& AllowedObjects)
+			: AllowedObjects(AllowedObjects)
+		{}
+
+		/** @return Whether Object is allowed to be replicated. */
+		bool IsObjectAllowed(const FConcertObjectInStreamID& Object) const { return AllowedObjects.Contains(Object); }
+		/** @return Number of allowed objects */
+		uint32 Num() const { return AllowedObjects.Num(); }
+
+		/** Enumerates all sync controlled objects. */
+		template<typename TProcessObject>
+		bool EnumerateAllowedObjects(TProcessObject Callback) const requires std::is_invocable_r_v<EBreakBehavior, TProcessObject, const FConcertObjectInStreamID&>;
+
+		/** Enumerates the changes that would be made by applying this event. */
+		template<CObjectInStreamInvocable TOnAllowed, CObjectInStreamInvocable TOnDisallowed>
+		static void EnumerateChanges(const FConcertReplication_ChangeSyncControl& Event, TOnAllowed&& OnAllowed, TOnDisallowed&& OnDisallowed);
+		
+		/** Combines past state with an incoming sync control event. */
+		template<CObjectInStreamInvocable TOnAllowed, CObjectInStreamInvocable TOnDisallowed>
+		void AppendChanges(const FConcertReplication_ChangeSyncControl& Event, TOnAllowed&& OnAllowed, TOnDisallowed&& OnDisallowed);
+		void AppendChanges(const FConcertReplication_ChangeSyncControl& Event);
+
+		/** Combines implicit and explicit sync control changes caused by authority changes. */
+		template<CObjectInStreamInvocable TOnAllowed, CObjectInStreamInvocable TOnDisallowed>
+		void AppendChanges(const FConcertReplication_ChangeAuthority_Request& Request, const FConcertReplication_ChangeAuthority_Response& Response, TOnAllowed&& OnAllowed, TOnDisallowed&& OnDisallowed);
+		void AppendChanges(const FConcertReplication_ChangeAuthority_Request& Request, const FConcertReplication_ChangeAuthority_Response& Response);
+
+		/** Combines implicit sync control changes caused by stream change. */
+		template<CObjectInStreamInvocable TOnDisallowed>
+		void AppendChanges(const FConcertReplication_ChangeStream_Request& Request, TOnDisallowed&& OnDisallowed);
+		void AppendChanges(const FConcertReplication_ChangeStream_Request& Request);
+		
+		friend bool operator==(const FSyncControlState& Left, const FSyncControlState& Right)
+		{
+			return Left.AllowedObjects.Num() == Right.AllowedObjects.Num() && Left.AllowedObjects.Includes(Right.AllowedObjects);
+		}
+		friend bool operator!=(const FSyncControlState& Left, const FSyncControlState& Right) { return !(Left == Right);}
+
+	private:
+
+		/** The objects the client has sync control over, i.e. is allowed to replicate. */
+		TSet<FConcertObjectInStreamID> AllowedObjects;
+	};
+	
+	template<typename TProcessObject>
+	bool FSyncControlState::EnumerateAllowedObjects(TProcessObject Callback) const requires std::is_invocable_r_v<EBreakBehavior, TProcessObject, const FConcertObjectInStreamID&>
+	{
+		for (const FConcertObjectInStreamID& Object : AllowedObjects)
+		{
+			if (Callback(Object) == EBreakBehavior::Break)
+			{
+				break;
+			}
+		}
+
+		return AllowedObjects.Num() > 0;
+	}
+
+	template <CObjectInStreamInvocable TOnAllowed, CObjectInStreamInvocable TOnDisallowed>
+	void FSyncControlState::EnumerateChanges(const FConcertReplication_ChangeSyncControl& Event, TOnAllowed&& OnAllowed, TOnDisallowed&& OnDisallowed)
+	{
+		for (const TPair<FConcertObjectInStreamID, bool>& NewControlState : Event.NewControlStates)
+		{
+			const FConcertObjectInStreamID& Object = NewControlState.Key;
+			if (NewControlState.Value)
+			{
+				OnAllowed(Object);
+			}
+			else
+			{
+				OnDisallowed(Object);
+			}
+		}
+	}
+
+	template<CObjectInStreamInvocable TOnAllowed, CObjectInStreamInvocable TOnDisallowed>
+	void FSyncControlState::AppendChanges(const FConcertReplication_ChangeSyncControl& Event, TOnAllowed&& OnAllowed, TOnDisallowed&& OnDisallowed)
+	{
+		EnumerateChanges(Event,
+			[this, &OnAllowed](const FConcertObjectInStreamID& Object){ AllowedObjects.Add(Object); OnAllowed(Object); },
+			[this, &OnDisallowed](const FConcertObjectInStreamID& Object){ AllowedObjects.Remove(Object); OnDisallowed(Object); }
+			);
+	}
+
+	inline void FSyncControlState::AppendChanges(const FConcertReplication_ChangeSyncControl& Event)
+	{
+		AppendChanges(Event, [](const FConcertObjectInStreamID&){}, [](const FConcertObjectInStreamID&){});
+	}
+	
+	template <CObjectInStreamInvocable TOnAllowed, CObjectInStreamInvocable TOnDisallowed>
+	void FSyncControlState::AppendChanges(const FConcertReplication_ChangeAuthority_Request& Request, const FConcertReplication_ChangeAuthority_Response& Response, TOnAllowed&& OnAllowed, TOnDisallowed&& OnDisallowed)
+	{
+		for (const TPair<FSoftObjectPath, FConcertStreamArray>& ImplicitChange : Request.ReleaseAuthority)
+		{
+			for (const FGuid& StreamId : ImplicitChange.Value.StreamIds)
+			{
+				const FConcertObjectInStreamID ImplicitlyRemoved { StreamId, ImplicitChange.Key };
+
+				const int32 NumRemoved = AllowedObjects.Remove(ImplicitlyRemoved);
+				if (NumRemoved > 0)
+				{
+					OnDisallowed(ImplicitlyRemoved);
+				}
+			}
+		}
+
+		AppendChanges(Response.SyncControl, MoveTemp(OnAllowed), MoveTemp(OnDisallowed));
+	}
+
+	inline void FSyncControlState::AppendChanges(const FConcertReplication_ChangeAuthority_Request& Request, const FConcertReplication_ChangeAuthority_Response& Response)
+	{
+		AppendChanges(Request, Response, [](const FConcertObjectInStreamID&){}, [](const FConcertObjectInStreamID&){});
+	}
+	
+	template <CObjectInStreamInvocable TOnDisallowed>
+	void FSyncControlState::AppendChanges(const FConcertReplication_ChangeStream_Request& Request, TOnDisallowed&& OnDisallowed)
+	{
+		for (auto It = AllowedObjects.CreateIterator(); It; ++It)
+		{
+			if (Request.StreamsToRemove.Contains(It->StreamId))
+			{
+				const FConcertObjectInStreamID StreamId = *It;
+				It.RemoveCurrent();
+				OnDisallowed(StreamId);
+			}
+		}
+		
+		if (AllowedObjects.IsEmpty())
+		{
+			return;
+		}
+
+		for (const FConcertObjectInStreamID& RemovedObject : Request.ObjectsToRemove)
+		{
+			if (IsObjectAllowed(RemovedObject))
+			{
+				AllowedObjects.Remove(RemovedObject);
+				OnDisallowed(RemovedObject);
+
+				if (AllowedObjects.IsEmpty())
+				{
+					break;
+				}
+			}
+		}
+	}
+
+	inline void FSyncControlState::AppendChanges(const FConcertReplication_ChangeStream_Request& Request)
+	{
+		AppendChanges(Request, [](const FConcertObjectInStreamID&){});
+	}
+}
+
diff --git a/Engine/Plugins/Developer/Concert/ConcertSync/ConcertSyncServer/Source/ConcertSyncServer/Private/Replication/AuthorityManager.cpp b/Engine/Plugins/Developer/Concert/ConcertSync/ConcertSyncServer/Source/ConcertSyncServer/Private/Replication/AuthorityManager.cpp
--- a/Engine/Plugins/Developer/Concert/ConcertSync/ConcertSyncServer/Source/ConcertSyncServer/Private/Replication/AuthorityManager.cpp
+++ b/Engine/Plugins/Developer/Concert/ConcertSync/ConcertSyncServer/Source/ConcertSyncServer/Private/Replication/AuthorityManager.cpp
@@ -16,10 +16,10 @@ namespace UE::ConcertSyncServer::Replication
 		class FServerGroundTruth : public ConcertSyncCore::Replication::AuthorityConflictUtils::IReplicationGroundTruth
 		{
 			const FAuthorityManager& Owner;
-			const IAuthorityManagerGetters& Getters;
+			const IRegistrationEnumerator& Getters;
 		public:
 			
-			FServerGroundTruth(const FAuthorityManager& Owner, const IAuthorityManagerGetters& Getters)
+			FServerGroundTruth(const FAuthorityManager& Owner, const IRegistrationEnumerator& Getters)
 				: Owner(Owner)
 				, Getters(Getters)
 			{}
@@ -34,7 +34,7 @@ namespace UE::ConcertSyncServer::Replication
 			
 			virtual void ForEachSendingClient(TFunctionRef<EBreakBehavior(const FGuid& ClientEndpointId)> Callback) const override
 			{
-				Getters.ForEachSendingClient(Callback);
+				Getters.ForEachReplicationClient(Callback);
 			}
 			
 			virtual bool HasAuthority(const FGuid& ClientId, const FGuid& StreamId, const FSoftObjectPath& ObjectPath) const override
@@ -58,8 +58,11 @@ namespace UE::ConcertSyncServer::Replication
 		}
 	}
 	
-	FAuthorityManager::FAuthorityManager(IAuthorityManagerGetters& Getters, TSharedRef<IConcertSession> InSession)
-		: Getters(Getters)
+	FAuthorityManager::FAuthorityManager(
+		IRegistrationEnumerator& InGetters,
+		TSharedRef<IConcertSession> InSession
+		)
+		: Getters(InGetters)
 		, Session(MoveTemp(InSession))
 	{
 		Session->RegisterCustomRequestHandler<FConcertReplication_ChangeAuthority_Request, FConcertReplication_ChangeAuthority_Response>(this, &FAuthorityManager::HandleChangeAuthorityRequest);
@@ -171,23 +174,25 @@ namespace UE::ConcertSyncServer::Replication
 	}
 
 	EConcertSessionResponseCode FAuthorityManager::HandleChangeAuthorityRequest(
-		const FConcertSessionContext& ConcertSessionContext,
+		const FConcertSessionContext& Context,
 		const FConcertReplication_ChangeAuthority_Request& Request,
 		FConcertReplication_ChangeAuthority_Response& Response
 		)
 	{
 		// This log does two things: 1. Identify issues in unit tests / at runtime 2. Warn about possibly malicious attempts when the server runs.
 		UE_CLOG(Request.TakeAuthority.IsEmpty() && Request.ReleaseAuthority.IsEmpty(), LogConcert, Warning, TEXT("Received invalid authority request (TakeAuthority.Num() == 0 and ReleaseAuthority.Num() == 0)"));
+		bool bMadeChanges = false;
 		
-		FClientAuthorityData& AuthorityData = ClientAuthorityData.FindOrAdd(ConcertSessionContext.SourceEndpointId);
-		const FClientId& ClientId = ConcertSessionContext.SourceEndpointId;
-		Private::ForEachReplicatedObject(Request.TakeAuthority, [this, &Response, &AuthorityData, &ClientId](const FStreamId& StreamId, const FSoftObjectPath& ObjectPath)
+		FClientAuthorityData& AuthorityData = ClientAuthorityData.FindOrAdd(Context.SourceEndpointId);
+		const FClientId& ClientId = Context.SourceEndpointId;
+		Private::ForEachReplicatedObject(Request.TakeAuthority, [this, &Response, &bMadeChanges, &AuthorityData, &ClientId](const FStreamId& StreamId, const FSoftObjectPath& ObjectPath)
 		{
 			const FConcertReplicatedObjectId ObjectToAuthor{ { StreamId, ObjectPath }, ClientId };
 			if (CanTakeAuthority(ObjectToAuthor))
 			{
 				UE_LOG(LogConcert, Log, TEXT("Transferred authority of %s to client %s for their stream %s"), *ObjectPath.ToString(), *ClientId.ToString(EGuidFormats::Short), *StreamId.ToString(EGuidFormats::Short));
 				AuthorityData.OwnedObjects.FindOrAdd(StreamId).Add(ObjectPath);
+				bMadeChanges = true;
 			}
 			else
 			{
@@ -196,7 +201,7 @@ namespace UE::ConcertSyncServer::Replication
 			}
 		});
 		
-		Private::ForEachReplicatedObject(Request.ReleaseAuthority, [this, &AuthorityData](const FStreamId& StreamId, const FSoftObjectPath& ObjectPath)
+		Private::ForEachReplicatedObject(Request.ReleaseAuthority, [this, &bMadeChanges, &AuthorityData](const FStreamId& StreamId, const FSoftObjectPath& ObjectPath)
 		{
 			TSet<FSoftObjectPath>* OwnedObjects = AuthorityData.OwnedObjects.Find(StreamId);
 			// Though dubious, it is a valid request for the client to release non-owned objects
@@ -210,10 +215,16 @@ namespace UE::ConcertSyncServer::Replication
 			if (OwnedObjects->IsEmpty())
 			{
 				AuthorityData.OwnedObjects.Remove(StreamId);
+				bMadeChanges = true;
 			}
 		});
 
 		Response.ErrorCode = EReplicationResponseErrorCode::Handled;
+		if (bMadeChanges && ensure(GenerateSyncControlDelegate.IsBound()))
+		{
+			Response.SyncControl = GenerateSyncControlDelegate.Execute(Context.SourceEndpointId);
+		}
+		
 		return EConcertSessionResponseCode::Success;
 	}
 	
diff --git a/Engine/Plugins/Developer/Concert/ConcertSync/ConcertSyncServer/Source/ConcertSyncServer/Private/Replication/AuthorityManager.h b/Engine/Plugins/Developer/Concert/ConcertSync/ConcertSyncServer/Source/ConcertSyncServer/Private/Replication/AuthorityManager.h
--- a/Engine/Plugins/Developer/Concert/ConcertSync/ConcertSyncServer/Source/ConcertSyncServer/Private/Replication/AuthorityManager.h
+++ b/Engine/Plugins/Developer/Concert/ConcertSync/ConcertSyncServer/Source/ConcertSyncServer/Private/Replication/AuthorityManager.h
@@ -3,10 +3,14 @@
 #pragma once
 
 #include "ConcertMessages.h"
-#include "Misc/EBreakBehavior.h"
+#include "Enumeration/IRegistrationEnumerator.h"
 #include "Replication/Data/ObjectIds.h"
+
+#include "Delegates/Delegate.h"
+#include "Misc/EBreakBehavior.h"
 #include "Templates/Function.h"
 
+struct FConcertReplication_ChangeSyncControl;
 class IConcertSession;
 
 struct FConcertReplication_ChangeAuthority_Response;
@@ -20,23 +24,6 @@ struct FConcertReplicationStream;
 
 namespace UE::ConcertSyncServer::Replication
 {
-	/**
-	 * Callbacks FAuthorityManager requires to function correctly.
-	 * Exists to make FAuthorityManager independent from all the other systems and allows mocking in tests.
-	 */
-	class IAuthorityManagerGetters
-	{
-	public:
-
-		/** Provides a way to extract all streams registered to a given client. */
-		virtual void ForEachStream(const FGuid& ClientEndpointId, TFunctionRef<EBreakBehavior(const FConcertReplicationStream& Stream)> Callback) const = 0;
-
-		/** Iterates through all clients have registered to send any data. */
-		virtual void ForEachSendingClient(TFunctionRef<EBreakBehavior(const FGuid& ClientEndpointId)> Callback) const = 0;
-
-		virtual ~IAuthorityManagerGetters() = default;
-	};
-	
 	/** Responds to FConcertChangeAuthority_Request and tracks what objects and properties clients have authority over. */
 	class FAuthorityManager : public FNoncopyable
 	{
@@ -45,8 +32,14 @@ namespace UE::ConcertSyncServer::Replication
 		using FStreamId = FGuid;
 		using FClientId = FGuid;
 		using FProcessAuthorityConflict = TFunctionRef<EBreakBehavior(const FClientId& ClientId, const FStreamId& StreamId, const FConcertPropertyChain& WrittenProperties)>;
+		
 
-		FAuthorityManager(IAuthorityManagerGetters& Getters, TSharedRef<IConcertSession> InSession);
+		/**
+		 * @param InGetters Gets information about clients' registered streams 
+		 * @param InSession The session to handle authority requests on
+		 * @param InGenerateSyncControlDelegate Called to fill the sync control portion of FConcertReplication_ChangeAuthority_Response.
+		 */
+		FAuthorityManager(IRegistrationEnumerator& InGetters UE_LIFETIMEBOUND, TSharedRef<IConcertSession> InSession);
 		~FAuthorityManager();
 
 		/**
@@ -84,11 +77,15 @@ namespace UE::ConcertSyncServer::Replication
 		void OnClientLeft(const FClientId& ClientEndpointId);
 		/** Takes away authority from the given client from the given object. */
 		void RemoveAuthority(const FConcertReplicatedObjectId& Object);
-
+		
+		DECLARE_DELEGATE_RetVal_OneParam(FConcertReplication_ChangeSyncControl, FGenerateSyncControl, const FGuid& ClientId);
+		/** Sets the delegate that is called to fill the sync control portion of FConcertReplication_ChangeAuthority_Response. */
+		FGenerateSyncControl& OnGenerateSyncControl() { return GenerateSyncControlDelegate; }
+		
 	private:
 
 		/** Callbacks required to obtain client info. */
-		IAuthorityManagerGetters& Getters;
+		IRegistrationEnumerator& Getters;
 		/** The session under which this manager operates. */
 		TSharedRef<IConcertSession> Session;
 
@@ -99,8 +96,11 @@ namespace UE::ConcertSyncServer::Replication
 		};
 		TMap<FClientId, FClientAuthorityData> ClientAuthorityData;
 
+		/** Called to fill the sync control portion of FConcertReplication_ChangeAuthority_Response. */
+		FGenerateSyncControl GenerateSyncControlDelegate;
+
 		EConcertSessionResponseCode HandleChangeAuthorityRequest(
-			const FConcertSessionContext& ConcertSessionContext,
+			const FConcertSessionContext& Context,
 			const FConcertReplication_ChangeAuthority_Request& Request,
 			FConcertReplication_ChangeAuthority_Response& Response
 			);
diff --git a/Engine/Plugins/Developer/Concert/ConcertSync/ConcertSyncServer/Source/ConcertSyncServer/Private/Replication/ConcertServerReplicationManager.cpp b/Engine/Plugins/Developer/Concert/ConcertSync/ConcertSyncServer/Source/ConcertSyncServer/Private/Replication/ConcertServerReplicationManager.cpp
--- a/Engine/Plugins/Developer/Concert/ConcertSync/ConcertSyncServer/Source/ConcertSyncServer/Private/Replication/ConcertServerReplicationManager.cpp
+++ b/Engine/Plugins/Developer/Concert/ConcertSync/ConcertSyncServer/Source/ConcertSyncServer/Private/Replication/ConcertServerReplicationManager.cpp
@@ -27,8 +27,9 @@ namespace UE::ConcertSyncServer::Replication
 		: Session(MoveTemp(InLiveSession))
 		, ReplicationFormat(MakeUnique<ConcertSyncCore::FFullObjectFormat>())
 		, AuthorityManager(*this, Session)
+		, SyncControlManager(*Session, AuthorityManager, *this)
 		, ReplicationCache(MakeShared<ConcertSyncCore::FObjectReplicationCache>(*ReplicationFormat))
-		, ReplicationDataReceiver(AuthorityManager, *Session, *ReplicationCache)
+		, ReplicationDataReceiver(AuthorityManager, SyncControlManager, *Session, *ReplicationCache)
 	{
 		Session->RegisterCustomRequestHandler<FConcertReplication_Join_Request, FConcertReplication_Join_Response>(this, &FConcertServerReplicationManager::HandleJoinReplicationSessionRequest);
 		Session->RegisterCustomRequestHandler<FConcertReplication_QueryReplicationInfo_Request, FConcertReplication_QueryReplicationInfo_Response>(this, &FConcertServerReplicationManager::HandleQueryReplicationInfoRequest);
@@ -65,12 +66,11 @@ namespace UE::ConcertSyncServer::Replication
 		}
 	}
 
-	void FConcertServerReplicationManager::ForEachSendingClient(TFunctionRef<EBreakBehavior(const FGuid& ClientEndpointId)> Callback) const
+	void FConcertServerReplicationManager::ForEachReplicationClient(TFunctionRef<EBreakBehavior(const FGuid& ClientEndpointId)> Callback) const
 	{
 		for (const TPair<FGuid, TUniquePtr<FConcertReplicationClient>>& ClientPair : Clients)
 		{
-			if (!ClientPair.Value->GetStreamDescriptions().IsEmpty()
-				&& Callback(ClientPair.Key) == EBreakBehavior::Break)
+			if (Callback(ClientPair.Key) == EBreakBehavior::Break)
 			{
 				break;
 			}
@@ -83,15 +83,22 @@ namespace UE::ConcertSyncServer::Replication
 		FConcertReplication_Join_Response& Response
 		)
 	{
+		const FGuid ClientId = ConcertSessionContext.SourceEndpointId;
 		// Have a pair of logs before and after processing in case of potential disaster
-		UE_LOG(LogConcert, Log, TEXT("Received replication join request from endpoint %s"), *ConcertSessionContext.SourceEndpointId.ToString());
+		UE_LOG(LogConcert, Log, TEXT("Received replication join request from endpoint %s"), *ClientId.ToString());
 		
-		LogNetworkMessage(CVarLogAuthorityRequestsAndResponsesOnServer, Request, [&](){ return GetClientName(*Session, ConcertSessionContext.SourceEndpointId); });
+		LogNetworkMessage(CVarLogAuthorityRequestsAndResponsesOnServer, Request, [&](){ return GetClientName(*Session, ClientId); });
 		const EConcertSessionResponseCode Result = InternalHandleJoinReplicationSessionRequest(ConcertSessionContext, Request, Response);
-		LogNetworkMessage(CVarLogAuthorityRequestsAndResponsesOnServer, Response, [&](){ return GetClientName(*Session, ConcertSessionContext.SourceEndpointId); });
+		LogNetworkMessage(CVarLogAuthorityRequestsAndResponsesOnServer, Response, [&](){ return GetClientName(*Session, ClientId); });
+
+		const bool bSuccess = Response.JoinErrorCode == EJoinReplicationErrorCode::Success;
+		if (bSuccess)
+		{
+			Response.SyncControl = SyncControlManager.OnGenerateSyncControlForClientJoin(ClientId);
+		}
 		
-		UE_CLOG(Response.JoinErrorCode == EJoinReplicationErrorCode::Success, LogConcert, Log, TEXT("Accepted replication join request"));
-		UE_CLOG(Response.JoinErrorCode != EJoinReplicationErrorCode::Success, LogConcert, Log, TEXT("Rejected replication join request. %s: %s"), *ConcertSyncCore::Replication::LexJoinErrorCode(Response.JoinErrorCode), *Response.DetailedErrorMessage);
+		UE_CLOG(bSuccess, LogConcert, Log, TEXT("Accepted replication join request"));
+		UE_CLOG(!bSuccess, LogConcert, Log, TEXT("Rejected replication join request. %s: %s"), *ConcertSyncCore::Replication::LexJoinErrorCode(Response.JoinErrorCode), *Response.DetailedErrorMessage);
 		return Result;
 	}
 
@@ -227,8 +234,7 @@ namespace UE::ConcertSyncServer::Replication
 		const FGuid ClientEndpointId = ConcertSessionContext.SourceEndpointId;
 		UE_LOG(LogConcert, Log, TEXT("Received replication leave request from endpoint %s"), *ClientEndpointId.ToString());
 		
-		Clients.Remove(ClientEndpointId);
-		AuthorityManager.OnClientLeft(ClientEndpointId);
+		OnClientLeftReplication(ClientEndpointId);
 	}
 
 	void FConcertServerReplicationManager::OnConnectionChanged(IConcertServerSession& ConcertServerSession, EConcertClientStatus ConcertClientStatus, const FConcertSessionClientInfo& ClientInfo)
@@ -236,11 +242,17 @@ namespace UE::ConcertSyncServer::Replication
 		const FGuid ClientEndpointId = ClientInfo.ClientEndpointId;
 		if (ConcertClientStatus == EConcertClientStatus::Disconnected)
 		{
-			Clients.Remove(ClientEndpointId);
-			AuthorityManager.OnClientLeft(ClientEndpointId);
+			OnClientLeftReplication(ClientEndpointId);
 		}
 	}
 
+	void FConcertServerReplicationManager::OnClientLeftReplication(const FGuid& EndpointId)
+	{
+		Clients.Remove(EndpointId);
+		AuthorityManager.OnClientLeft(EndpointId);
+		SyncControlManager.OnClientLeft(EndpointId);
+	}
+
 	void FConcertServerReplicationManager::Tick(IConcertServerSession& InSession, float InDeltaTime)
 	{
 		// TODO UE-190714: We may want to set a time budget for clients.
diff --git a/Engine/Plugins/Developer/Concert/ConcertSync/ConcertSyncServer/Source/ConcertSyncServer/Private/Replication/ConcertServerReplicationManager.h b/Engine/Plugins/Developer/Concert/ConcertSync/ConcertSyncServer/Source/ConcertSyncServer/Private/Replication/ConcertServerReplicationManager.h
--- a/Engine/Plugins/Developer/Concert/ConcertSync/ConcertSyncServer/Source/ConcertSyncServer/Private/Replication/ConcertServerReplicationManager.h
+++ b/Engine/Plugins/Developer/Concert/ConcertSync/ConcertSyncServer/Source/ConcertSyncServer/Private/Replication/ConcertServerReplicationManager.h
@@ -5,11 +5,13 @@
 #include "AuthorityManager.h"
 #include "ConcertMessages.h"
 #include "ConcertReplicationClient.h"
+#include "Enumeration/IRegistrationEnumerator.h"
 #include "Replication/Formats/IObjectReplicationFormat.h"
 #include "Replication/IConcertServerReplicationManager.h"
 #include "Replication/Messages/Handshake.h"
 #include "Replication/Processing/ObjectReplicationCache.h"
 #include "Replication/Processing/ServerObjectReplicationReceiver.h"
+#include "SyncControlManager.h"
 
 #include "Templates/SharedPointer.h"
 #include "Templates/Tuple.h"
@@ -43,7 +45,7 @@ namespace UE::ConcertSyncServer::Replication
 	 */
 	class FConcertServerReplicationManager
 		: public IConcertServerReplicationManager
-		, public IAuthorityManagerGetters
+		, public IRegistrationEnumerator
 		, public FNoncopyable
 	{
 	public:
@@ -53,10 +55,13 @@ namespace UE::ConcertSyncServer::Replication
 
 		const FAuthorityManager& GetAuthorityManager() const { return AuthorityManager; }
 
-		//~ Begin IAuthorityManagerGetters Interface
+		//~ Begin IStreamEnumerator Interface
 		virtual void ForEachStream(const FGuid& ClientEndpointId, TFunctionRef<EBreakBehavior(const FConcertReplicationStream& Stream)> Callback) const override;
-		virtual void ForEachSendingClient(TFunctionRef<EBreakBehavior(const FGuid& ClientEndpointId)> Callback) const override;
-		//~ End IAuthorityManagerGetters Interface
+		//~ End IStreamEnumerator Interface
+		
+		//~ Begin IClientEnumerator Interface
+		virtual void ForEachReplicationClient(TFunctionRef<EBreakBehavior(const FGuid& ClientEndpointId)> Callback) const override;
+		//~ End IClientEnumerator Interface
 
 	private:
 		
@@ -68,6 +73,8 @@ namespace UE::ConcertSyncServer::Replication
 
 		/** Responds to client requests to changing authority and can be asked whether an object change is valid to take place. */
 		FAuthorityManager AuthorityManager;
+		/** Decides whether clients should be replicating. Clients may replicate when they have authority and there are other clients listening for that data. */
+		FSyncControlManager SyncControlManager;
 		
 		/** Received replication events are put into the ReplicationCache. The cache is used to relay data to clients latently. */
 		const TSharedRef<ConcertSyncCore::FObjectReplicationCache> ReplicationCache;
@@ -98,6 +105,8 @@ namespace UE::ConcertSyncServer::Replication
 		void HandleLeaveReplicationSessionRequest(const FConcertSessionContext& ConcertSessionContext, const FConcertReplication_LeaveEvent& EventData);
 		void OnConnectionChanged(IConcertServerSession& ConcertServerSession, EConcertClientStatus ConcertClientStatus, const FConcertSessionClientInfo& ConcertSessionClientInfo);
 
+		void OnClientLeftReplication(const FGuid& EndpointId);
+		
 		/**
 		 * Ticks all clients which causes clients to process pending data and send it to the corresponding endpoints.
 		 * 
diff --git a/Engine/Plugins/Developer/Concert/ConcertSync/ConcertSyncServer/Source/ConcertSyncServer/Private/Replication/Enumeration/IClientEnumerator.h b/Engine/Plugins/Developer/Concert/ConcertSync/ConcertSyncServer/Source/ConcertSyncServer/Private/Replication/Enumeration/IClientEnumerator.h
new file mode 100644
--- /dev/null
+++ b/Engine/Plugins/Developer/Concert/ConcertSync/ConcertSyncServer/Source/ConcertSyncServer/Private/Replication/Enumeration/IClientEnumerator.h
@@ -0,0 +1,23 @@
+﻿// Copyright Epic Games, Inc. All Rights Reserved.
+
+#pragma once
+
+#include "Misc/EBreakBehavior.h"
+#include "Templates/Function.h"
+
+struct FConcertReplicationStream;
+struct FGuid;
+
+namespace UE::ConcertSyncServer::Replication
+{
+	/** Iterates through all clients have registered to send any data. */
+	class IClientEnumerator
+	{
+	public:
+
+		/** Iterates through all clients have registered to send any data. */
+		virtual void ForEachReplicationClient(TFunctionRef<EBreakBehavior(const FGuid& ClientEndpointId)> Callback) const = 0;
+		
+		virtual ~IClientEnumerator() = default;
+	};
+}
\ No newline at end of file
diff --git a/Engine/Plugins/Developer/Concert/ConcertSync/ConcertSyncServer/Source/ConcertSyncServer/Private/Replication/Enumeration/IRegistrationEnumerator.h b/Engine/Plugins/Developer/Concert/ConcertSync/ConcertSyncServer/Source/ConcertSyncServer/Private/Replication/Enumeration/IRegistrationEnumerator.h
new file mode 100644
--- /dev/null
+++ b/Engine/Plugins/Developer/Concert/ConcertSync/ConcertSyncServer/Source/ConcertSyncServer/Private/Replication/Enumeration/IRegistrationEnumerator.h
@@ -0,0 +1,12 @@
+﻿// Copyright Epic Games, Inc. All Rights Reserved.
+
+#pragma once
+
+#include "IClientEnumerator.h"
+#include "IStreamEnumerator.h"
+
+namespace UE::ConcertSyncServer::Replication
+{
+	class IRegistrationEnumerator : public IStreamEnumerator, public IClientEnumerator
+	{};
+}
\ No newline at end of file
diff --git a/Engine/Plugins/Developer/Concert/ConcertSync/ConcertSyncServer/Source/ConcertSyncServer/Private/Replication/Enumeration/IStreamEnumerator.h b/Engine/Plugins/Developer/Concert/ConcertSync/ConcertSyncServer/Source/ConcertSyncServer/Private/Replication/Enumeration/IStreamEnumerator.h
new file mode 100644
--- /dev/null
+++ b/Engine/Plugins/Developer/Concert/ConcertSync/ConcertSyncServer/Source/ConcertSyncServer/Private/Replication/Enumeration/IStreamEnumerator.h
@@ -0,0 +1,23 @@
+﻿// Copyright Epic Games, Inc. All Rights Reserved.
+
+#pragma once
+
+#include "Misc/EBreakBehavior.h"
+#include "Templates/Function.h"
+
+struct FConcertReplicationStream;
+struct FGuid;
+
+namespace UE::ConcertSyncServer::Replication
+{
+	/** Provides a way to extract all streams registered to a given client. */
+	class IStreamEnumerator
+	{
+	public:
+
+		/** Provides a way to extract all streams registered to a given client. */
+		virtual void ForEachStream(const FGuid& ClientEndpointId, TFunctionRef<EBreakBehavior(const FConcertReplicationStream& Stream)> Callback) const = 0;
+
+		virtual ~IStreamEnumerator() = default;
+	};
+}
\ No newline at end of file
diff --git a/Engine/Plugins/Developer/Concert/ConcertSync/ConcertSyncServer/Source/ConcertSyncServer/Private/Replication/Processing/ServerObjectReplicationReceiver.cpp b/Engine/Plugins/Developer/Concert/ConcertSync/ConcertSyncServer/Source/ConcertSyncServer/Private/Replication/Processing/ServerObjectReplicationReceiver.cpp
--- a/Engine/Plugins/Developer/Concert/ConcertSync/ConcertSyncServer/Source/ConcertSyncServer/Private/Replication/Processing/ServerObjectReplicationReceiver.cpp
+++ b/Engine/Plugins/Developer/Concert/ConcertSync/ConcertSyncServer/Source/ConcertSyncServer/Private/Replication/Processing/ServerObjectReplicationReceiver.cpp
@@ -2,8 +2,10 @@
 
 #include "ServerObjectReplicationReceiver.h"
 
+#include "ConcertLogGlobal.h"
 #include "IConcertSessionHandler.h"
 #include "Replication/AuthorityManager.h"
+#include "Replication/SyncControlManager.h"
 #include "Replication/Data/ObjectIds.h"
 #include "Replication/Messages/ObjectReplication.h"
 
@@ -11,11 +13,13 @@ namespace UE::ConcertSyncServer::Replication
 {
 	FServerObjectReplicationReceiver::FServerObjectReplicationReceiver(
 		const FAuthorityManager& AuthorityManager,
+		const FSyncControlManager& SyncControlManager,
 		IConcertSession& Session,
 		ConcertSyncCore::FObjectReplicationCache& ReplicationCache
 		)
 		: FObjectReplicationReceiver(Session, ReplicationCache)
 		, AuthorityManager(AuthorityManager)
+		, SyncControlManager(SyncControlManager)
 	{}
 
 	bool FServerObjectReplicationReceiver::ShouldAcceptObject(
@@ -25,6 +29,13 @@ namespace UE::ConcertSyncServer::Replication
 		) const
 	{
 		const FConcertReplicatedObjectId ReplicatedObjectInfo { { StreamEvent.StreamId, ObjectEvent.ReplicatedObject }, SessionContext.SourceEndpointId };
-		return AuthorityManager.HasAuthorityToChange(ReplicatedObjectInfo);
+		
+		const bool bHasAuthority = AuthorityManager.HasAuthorityToChange(ReplicatedObjectInfo);
+		const bool bHasSyncControl = SyncControlManager.HasSyncControl(ReplicatedObjectInfo);
+		UE_CLOG(!bHasAuthority, LogConcert, Verbose, TEXT("Dropping %s because the client does not have authority over it."), *ReplicatedObjectInfo.ToString());
+		UE_CLOG(!bHasSyncControl, LogConcert, Verbose, TEXT("Dropping %s because the client does not have sync control over it."), *ReplicatedObjectInfo.ToString());
+
+		// Note: Having sync control logically implies that the sender should have authority but we check both conditions here for completeness.
+		return bHasAuthority && bHasSyncControl;
 	}
 }
diff --git a/Engine/Plugins/Developer/Concert/ConcertSync/ConcertSyncServer/Source/ConcertSyncServer/Private/Replication/Processing/ServerObjectReplicationReceiver.h b/Engine/Plugins/Developer/Concert/ConcertSync/ConcertSyncServer/Source/ConcertSyncServer/Private/Replication/Processing/ServerObjectReplicationReceiver.h
--- a/Engine/Plugins/Developer/Concert/ConcertSync/ConcertSyncServer/Source/ConcertSyncServer/Private/Replication/Processing/ServerObjectReplicationReceiver.h
+++ b/Engine/Plugins/Developer/Concert/ConcertSync/ConcertSyncServer/Source/ConcertSyncServer/Private/Replication/Processing/ServerObjectReplicationReceiver.h
@@ -4,9 +4,12 @@
 
 #include "Replication/Processing/ObjectReplicationReceiver.h"
 
+#include "HAL/Platform.h"
+
 namespace UE::ConcertSyncServer::Replication
 {
 	class FAuthorityManager;
+	class FSyncControlManager;
 	
 	/** Rejects changes to objects that the sending client does not have authority over. */
 	class FServerObjectReplicationReceiver : public ConcertSyncCore::FObjectReplicationReceiver
@@ -15,6 +18,7 @@ namespace UE::ConcertSyncServer::Replication
 
 		FServerObjectReplicationReceiver(
 			const FAuthorityManager& AuthorityManager UE_LIFETIMEBOUND,
+			const FSyncControlManager& SyncControlManager UE_LIFETIMEBOUND,
 			IConcertSession& Session UE_LIFETIMEBOUND,
 			ConcertSyncCore::FObjectReplicationCache& ReplicationCache UE_LIFETIMEBOUND
 			);
@@ -29,6 +33,8 @@ namespace UE::ConcertSyncServer::Replication
 
 		/** Used to determine whether a client has authority over objects. */
 		const FAuthorityManager& AuthorityManager;
+		/** Used to determine whether any client is listening to an incoming object. */
+		const FSyncControlManager& SyncControlManager;
 	};
 }
 
diff --git a/Engine/Plugins/Developer/Concert/ConcertSync/ConcertSyncServer/Source/ConcertSyncServer/Private/Replication/SyncControlManager.cpp b/Engine/Plugins/Developer/Concert/ConcertSync/ConcertSyncServer/Source/ConcertSyncServer/Private/Replication/SyncControlManager.cpp
new file mode 100644
--- /dev/null
+++ b/Engine/Plugins/Developer/Concert/ConcertSync/ConcertSyncServer/Source/ConcertSyncServer/Private/Replication/SyncControlManager.cpp
@@ -0,0 +1,151 @@
+﻿// Copyright Epic Games, Inc. All Rights Reserved.
+
+#include "SyncControlManager.h"
+
+#include "AuthorityManager.h"
+#include "ConcertServer/Private/ConcertServerSession.h"
+#include "Replication/Data/ReplicationStream.h"
+#include "Replication/Messages/SyncControl.h"
+
+namespace UE::ConcertSyncServer::Replication
+{
+	FSyncControlManager::FSyncControlManager(
+		IConcertServerSession& Session,
+		FAuthorityManager& AuthorityManager,
+		const IRegistrationEnumerator& Getters
+		)
+		: Session(Session)
+		, AuthorityManager(AuthorityManager)
+		, Getters(Getters)
+	{
+		AuthorityManager.OnGenerateSyncControl().BindRaw(this, &FSyncControlManager::OnGenerateSyncControlForAuthorityResponse);
+	}
+
+	FSyncControlManager::~FSyncControlManager()
+	{
+		Session.OnSessionClientChanged().RemoveAll(this);
+		AuthorityManager.OnGenerateSyncControl().Unbind();
+	}
+
+	bool FSyncControlManager::HasSyncControl(const FConcertReplicatedObjectId& Object) const
+	{
+		const FClientData* ClientData = PerClientData.Find(Object.SenderEndpointId);
+		return ClientData && ClientData->ObjectsWithSyncControl.Contains(Object);
+	}
+
+	FConcertReplication_ChangeSyncControl FSyncControlManager::OnGenerateSyncControlForClientJoin(const FGuid& ClientId)
+	{
+		RefreshAndSendToAllClientsExcept(ClientId);
+		return RefreshClientSyncControl(ClientId);
+	}
+
+	FConcertReplication_ChangeSyncControl FSyncControlManager::OnGenerateSyncControlForAuthorityResponse(const FGuid& ClientId)
+	{
+		RefreshAndSendToAllClientsExcept(ClientId);
+
+		const auto ShouldSkipInControlMessage = [](const FConcertObjectInStreamID&, bool bNewState)
+		{
+			// We skip including states that disable sync control to save on network bandwidth.
+			// The receiving client can already predict the objects that lose sync control based on their request.
+			// See FConcertReplication_ChangeSyncControl and FConcertReplication_ChangeAuthority_Response::SyncControl documentation.
+			return !bNewState;
+		};
+		return RefreshClientSyncControl(ClientId, ShouldSkipInControlMessage);
+	}
+
+	void FSyncControlManager::HandleClientLeave(const FGuid& LeftClientId)
+	{
+		// The disconnecting client's objects should be removed
+		PerClientData.Remove(LeftClientId);
+
+		// The leaving client may have been the last client to listen for certain object updates
+		RefreshAndSendToAllClientsExcept(LeftClientId);
+	}
+	
+	void FSyncControlManager::RefreshAndSendSyncControl(const FGuid& ClientId)
+	{
+		FConcertReplication_ChangeSyncControl SyncControlChange = RefreshClientSyncControl(ClientId);
+		if (!SyncControlChange.IsEmpty())
+		{
+			Session.SendCustomEvent(SyncControlChange, { ClientId }, EConcertMessageFlags::ReliableOrdered);
+		}
+	}
+	
+	FConcertReplication_ChangeSyncControl FSyncControlManager::RefreshClientSyncControl(
+		const FGuid& ClientId,
+		TFunctionRef<bool(const FConcertObjectInStreamID& Object, bool bNewValue)> ShouldSkipInMessage
+		)
+	{
+		FConcertReplication_ChangeSyncControl SyncControlChange;
+		const bool bHadSyncControlBefore = PerClientData.Contains(ClientId);
+		FClientData& ClientData = PerClientData.FindOrAdd(ClientId);
+
+		Getters.ForEachStream(ClientId,
+			[this, &ClientId, &ShouldSkipInMessage, bHadSyncControlBefore, &ClientData, &SyncControlChange]
+			(const FConcertReplicationStream& Stream)
+			{
+				const FConcertBaseStreamInfo& BaseInfo = Stream.BaseDescription;
+				for (const TPair<FSoftObjectPath, FConcertReplicatedObjectInfo>& Pair : BaseInfo.ReplicationMap.ReplicatedObjects)
+				{
+					const FConcertObjectInStreamID ObjectId { BaseInfo.Identifier, Pair.Key };
+					const FConcertReplicatedObjectId ReplicatedObjectId { ObjectId, ClientId };
+					const bool bIsValidForSending = Pair.Value.IsValidForSendingToServer()
+						&& AuthorityManager.HasAuthorityToChange(ReplicatedObjectId)
+						&& IsAnyoneInterestedIn(ReplicatedObjectId);
+					
+					const bool bIsSameAsBefore = ClientData.ObjectsWithSyncControl.Contains(ObjectId) == bIsValidForSending;
+					if (bIsValidForSending)
+					{
+						ClientData.ObjectsWithSyncControl.Add(ObjectId);
+					}
+					else
+					{
+						ClientData.ObjectsWithSyncControl.Remove(ObjectId);
+					}
+					
+					const bool bNeedsToIncludeInMessage = !bHadSyncControlBefore || !bIsSameAsBefore;
+					if (bNeedsToIncludeInMessage && !ShouldSkipInMessage(ObjectId, bIsValidForSending))
+					{
+						SyncControlChange.NewControlStates.Add(ObjectId, bIsValidForSending);
+					}
+				}
+
+				return EBreakBehavior::Continue;
+			});
+
+		return SyncControlChange;
+	}
+
+	void FSyncControlManager::RefreshAndSendToAllClientsExcept(const FGuid& SkippedClient)
+	{
+		Getters.ForEachReplicationClient([this, &SkippedClient](const FGuid& ClientId)
+		{
+			if (SkippedClient != ClientId)
+			{
+				RefreshAndSendSyncControl(ClientId);
+			}
+
+			return EBreakBehavior::Continue;
+		});
+	}
+	
+	bool FSyncControlManager::IsAnyoneInterestedIn(const FConcertReplicatedObjectId& Object) const
+	{
+		bool bFoundClient = false;
+		Getters.ForEachReplicationClient([this, &Object, &bFoundClient](const FGuid& ListenerId)
+		{
+			bFoundClient |= IsClientInterestedIn(Object, ListenerId);
+			return bFoundClient ? EBreakBehavior::Break : EBreakBehavior::Continue;
+		});
+		return bFoundClient;
+	}
+
+	bool FSyncControlManager::IsClientInterestedIn(const FConcertReplicatedObjectId& Object, const FGuid& Client) const
+	{
+		// For now, we'll just prevent sending if there is no other client.
+		// TODO: In the future, we should consider whether Client in the same world as Object.SenderEndpointId
+		const bool bIsSameAsSender = Client == Object.SenderEndpointId;
+		return !bIsSameAsSender;
+	}
+}
+
diff --git a/Engine/Plugins/Developer/Concert/ConcertSync/ConcertSyncServer/Source/ConcertSyncServer/Private/Replication/SyncControlManager.h b/Engine/Plugins/Developer/Concert/ConcertSync/ConcertSyncServer/Source/ConcertSyncServer/Private/Replication/SyncControlManager.h
new file mode 100644
--- /dev/null
+++ b/Engine/Plugins/Developer/Concert/ConcertSync/ConcertSyncServer/Source/ConcertSyncServer/Private/Replication/SyncControlManager.h
@@ -0,0 +1,89 @@
+﻿// Copyright Epic Games, Inc. All Rights Reserved.
+
+#pragma once
+
+#include "ConcertMessageData.h"
+#include "Replication/Data/ObjectIds.h"
+
+#include "HAL/Platform.h"
+#include "Replication/Messages/SyncControl.h"
+#include "Templates/UnrealTemplate.h"
+
+class IConcertServerSession;
+enum class EConcertClientStatus : uint8;
+struct FConcertSessionClientInfo;
+
+namespace UE::ConcertSyncServer::Replication
+{
+	class IRegistrationEnumerator;
+	class FAuthorityManager;
+	
+	/**
+	 * Decides whether clients should be replicating.
+	 * Clients may replicate when they have authority and there are other clients listening for that data.
+	 *
+	 * For now, sync control just checks whether there is another client in the session.
+	 * TODO: In the future, we should consider whether there is another client in the same world, as well.
+	 */
+	class FSyncControlManager : public FNoncopyable
+	{
+	public:
+		
+		FSyncControlManager(
+			IConcertServerSession& Session UE_LIFETIMEBOUND,
+			FAuthorityManager& AuthorityManager UE_LIFETIMEBOUND,
+			const IRegistrationEnumerator& Getters UE_LIFETIMEBOUND
+			);
+		~FSyncControlManager();
+
+		/** @return Whether this Object is allowed to processed. */
+		bool HasSyncControl(const FConcertReplicatedObjectId& Object) const;
+
+		/** Called by FConcertServerReplicationManager when client completes replication handshake. */
+		FConcertReplication_ChangeSyncControl OnGenerateSyncControlForClientJoin(const FGuid& ClientId);
+		
+		/** Called by FConcertServerReplicationManager when client completes leaves replication. */
+		void OnClientLeft(const FGuid& ClientId) { HandleClientLeave(ClientId); }
+
+	private:
+
+		/** Used to send sync control messages to clients, which notifies them to start / stop replicating. */
+		IConcertServerSession& Session;
+		/** Used to detect whether a client has authority. */
+		FAuthorityManager& AuthorityManager;
+
+		/** Callbacks for retrieving more info about client replication registration. */
+		const IRegistrationEnumerator& Getters;
+
+		struct FClientData
+		{
+			/** Contains all objects that are allowed to be replicated for this client. */
+			TSet<FConcertObjectInStreamID> ObjectsWithSyncControl;
+		};
+		/** Maps client ID to client sync control data. */
+		TMap<FGuid, FClientData> PerClientData;
+
+		/** Generates new sync control for client that is explicitly changing their authority. */
+		FConcertReplication_ChangeSyncControl OnGenerateSyncControlForAuthorityResponse(const FGuid& ClientId);
+
+		/** Cleans up the associated client data and updates sync control for all other clients. */
+		void HandleClientLeave(const FGuid& LeftClientId);
+
+		/** Compares the client's current sync control against the new sync control it should have and returns the delta to be sent to the client. */
+		FConcertReplication_ChangeSyncControl RefreshClientSyncControl(
+			const FGuid& ClientId, 
+			TFunctionRef<bool(const FConcertObjectInStreamID& Object, bool bNewValue)> ShouldSkipInMessage = [](auto&, auto){ return false; }
+			);
+
+		/** Except for SkippedClient, checks all clients' sync control has changed and conditionally updates the remote endpoints. */
+		void RefreshAndSendToAllClientsExcept(const FGuid& SkippedClient);
+		/** Checks whether client sync control has changed and conditionally updates the remote endpoint. */
+		void RefreshAndSendSyncControl(const FGuid& ClientId);
+
+		/** @return Whether any client wants to receive Object. */
+		bool IsAnyoneInterestedIn(const FConcertReplicatedObjectId& Object) const;
+		/** @return Whether Client is interested in receiving Object. */
+		bool IsClientInterestedIn(const FConcertReplicatedObjectId& Object, const FGuid& Client) const;
+	};
+}
+
diff --git a/Engine/Plugins/Developer/Concert/ConcertSync/ConcertSyncTest/Source/ConcertSyncTest/Private/Replication/Components/SyncControlState.spec.cpp b/Engine/Plugins/Developer/Concert/ConcertSync/ConcertSyncTest/Source/ConcertSyncTest/Private/Replication/Components/SyncControlState.spec.cpp
new file mode 100644
--- /dev/null
+++ b/Engine/Plugins/Developer/Concert/ConcertSync/ConcertSyncTest/Source/ConcertSyncTest/Private/Replication/Components/SyncControlState.spec.cpp
@@ -0,0 +1,84 @@
+﻿// Copyright Epic Games, Inc. All Rights Reserved.
+
+#include "Misc/AutomationTest.h"
+
+#include "Replication/SyncControlState.h"
+#include "Replication/Messages/SyncControl.h"
+#include "Replication/Util/Spec/ObjectTestReplicator.h"
+
+namespace UE::ConcertSyncTests::Replication::UI
+{
+	BEGIN_DEFINE_SPEC(FSyncControlStateSpec, "Editor.Concert.Replication.Components", EAutomationTestFlags::EditorContext | EAutomationTestFlags::EngineFilter)
+	FGuid SenderStreamId = FGuid::NewGuid();
+	/** Leverage to create a real FSoftObjectPath. */
+	TSharedPtr<FObjectTestReplicator> ObjectReplicator;
+	END_DEFINE_SPEC(FSyncControlStateSpec);
+
+	/** This tests that FSyncControlState correctly analyses requests and responses for aggregation. */
+	void FSyncControlStateSpec::Define()
+	{
+		BeforeEach([this]
+		{
+			ObjectReplicator = MakeShared<FObjectTestReplicator>();
+		});
+		AfterEach([this]
+		{
+			// Test would hold onto this for rest of engine lifetime. Clean up this mini would-be leak.
+			ObjectReplicator.Reset();
+		});
+
+		// At this point, everything should work server-side. Now test client-side prediction.
+		Describe("Changing sync control state", [this]()
+		{
+			It("Is correct with releasing authority with FConcertReplication_ChangeAuthority_Request", [this]()
+			{
+				ConcertSyncCore::Replication::FSyncControlState SyncControl = TSet<FConcertObjectInStreamID>{{ SenderStreamId, ObjectReplicator->TestObject }};
+				const FConcertReplication_ChangeAuthority_Request Request { .ReleaseAuthority = { { ObjectReplicator->TestObject, {{ SenderStreamId }} } } };
+				const FConcertReplication_ChangeAuthority_Response Response;
+				
+				TArray<FConcertObjectInStreamID> RemovedObjects;
+				SyncControl.AppendChanges(
+					Request,
+					Response,
+					[this](const FConcertObjectInStreamID&){ AddError(TEXT("No object should be added")); },
+					[&RemovedObjects](const FConcertObjectInStreamID& Object){ RemovedObjects.Add(Object); }
+					);
+				
+				TestTrue(TEXT("Removed TestObject"), RemovedObjects.Contains(FConcertObjectInStreamID{ SenderStreamId, ObjectReplicator->TestObject }));
+				TestEqual(TEXT("Removed exactly 1 object"), RemovedObjects.Num(), 1);
+			});
+			
+			It("Is correct with taking authority with FConcertReplication_ChangeAuthority_Request", [this]()
+			{
+				ConcertSyncCore::Replication::FSyncControlState SyncControl;
+				const FConcertObjectInStreamID ObjectId{ SenderStreamId, ObjectReplicator->TestObject };
+				const FConcertReplication_ChangeAuthority_Request Request { .TakeAuthority = {{ ObjectReplicator->TestObject, {{ SenderStreamId }}}}};
+				const FConcertReplication_ChangeAuthority_Response Response { .SyncControl = {{{ ObjectId, true }}}};
+				
+				TArray<FConcertObjectInStreamID> AddedObjects;
+				SyncControl.AppendChanges(
+					Request,
+					Response,
+					[&AddedObjects](const FConcertObjectInStreamID& Object){ AddedObjects.Add(Object); },
+					[this](const FConcertObjectInStreamID&){ AddError(TEXT("No object should be removed")); }
+					);
+				
+				TestTrue(TEXT("Added TestObject"), AddedObjects.Contains(ObjectId));
+				TestEqual(TEXT("Added exactly 1 object"), AddedObjects.Num(), 1);
+			});
+
+			It("Is correct after releasing authority with FConcertReplication_ChangeStream_Request", [this]()
+			{
+				ConcertSyncCore::Replication::FSyncControlState SyncControl = TSet<FConcertObjectInStreamID>{{ SenderStreamId, ObjectReplicator->TestObject }};
+				FConcertReplication_ChangeStream_Request StreamChange;
+				StreamChange.ObjectsToRemove.Add({ SenderStreamId, ObjectReplicator->TestObject });
+
+				TArray<FConcertObjectInStreamID> RemovedObjects;
+				SyncControl.AppendChanges(StreamChange, [&RemovedObjects](const FConcertObjectInStreamID& Object){ RemovedObjects.Add(Object); });
+				
+				TestTrue(TEXT("Removed TestObject"), RemovedObjects.Contains(FConcertObjectInStreamID{ SenderStreamId, ObjectReplicator->TestObject }));
+				TestEqual(TEXT("Removed exactly 1 object"), RemovedObjects.Num(), 1);
+			});
+		});
+	}
+}
diff --git a/Engine/Plugins/Developer/Concert/ConcertSync/ConcertSyncTest/Source/ConcertSyncTest/Private/Replication/SyncControl.spec.cpp b/Engine/Plugins/Developer/Concert/ConcertSync/ConcertSyncTest/Source/ConcertSyncTest/Private/Replication/SyncControl.spec.cpp
new file mode 100644
--- /dev/null
+++ b/Engine/Plugins/Developer/Concert/ConcertSync/ConcertSyncTest/Source/ConcertSyncTest/Private/Replication/SyncControl.spec.cpp
@@ -0,0 +1,326 @@
+﻿// Copyright Epic Games, Inc. All Rights Reserved.
+
+#include "Misc/AutomationTest.h"
+#include "Replication/SyncControlState.h"
+#include "Replication/Messages/SyncControl.h"
+#include "Util/ConcertClientReplicationBridgeMock.h"
+#include "Util/Spec/ObjectTestReplicator.h"
+#include "Util/Spec/ReplicationClient.h"
+#include "Util/Spec/ReplicationServer.h"
+
+namespace UE::ConcertSyncTests::Replication::UI
+{
+	BEGIN_DEFINE_SPEC(FReplicationSyncControlSpec, "Editor.Concert.Replication.SyncControl", EAutomationTestFlags::EditorContext | EAutomationTestFlags::EngineFilter)
+	
+	TSharedPtr<FReplicationServer> Server;
+	FReplicationClient* Sender = nullptr;
+	FReplicationClient* Receiver = nullptr;
+
+	FGuid SenderStreamId = FGuid::NewGuid();
+	TSharedPtr<FObjectTestReplicator> ObjectReplicator;
+
+	FObjectReplicationContext MakeSenderToReceiverContext() const { return { *Sender, *Server, *Receiver }; }
+	FObjectReplicationContext MakeSenderToReceiverContext(FReplicationClient& OtherReceiver) const { return { *Sender, *Server, OtherReceiver }; }
+	
+	END_DEFINE_SPEC(FReplicationSyncControlSpec);
+
+	/** This tests that clients respect sync control (see FConcertReplication_ChangeSyncControl). */
+	void FReplicationSyncControlSpec::Define()
+	{
+		BeforeEach([this]
+		{
+			Server = MakeShared<FReplicationServer>(*this);
+			Sender = &Server->ConnectClient();
+			Receiver = &Server->ConnectClient();
+
+			ObjectReplicator = MakeShared<FObjectTestReplicator>();
+
+			Sender->JoinReplication(ObjectReplicator->CreateSenderArgs(SenderStreamId));
+			// Tells client that the object is ready to replicate
+			Sender->GetBridgeMock().InjectAvailableObject(*ObjectReplicator->TestObject);
+			// Tells server intent to replicate the object
+			Sender->GetClientReplicationManager().TakeAuthorityOver({ ObjectReplicator->TestObject });
+		});
+		AfterEach([this]
+		{
+			// Test would hold onto this for rest of engine lifetime. Clean up this mini would-be leak.
+			Server.Reset();
+			ObjectReplicator.Reset();
+		});
+		
+		Describe("Before receiver client joins", [this]()
+		{
+			It("Client will not attempt to replicate", [this]()
+			{
+				ObjectReplicator->SimulateSendObjectToReceiver(*this, MakeSenderToReceiverContext(), { SenderStreamId },
+					[this](const FConcertSessionContext& Context, const FConcertReplication_BatchReplicationEvent& Event)
+					{
+						AddError(TEXT("Sending client replicated object even though they did not have permission."));
+					});
+				ObjectReplicator->TestValuesWereNotReplicated(*this);
+
+				// This will cause the server to send a sync control event
+				Receiver->JoinReplicationAsListener({ ObjectReplicator->TestObject });
+				
+				ObjectReplicator->SimulateSendObjectToReceiver(*this, MakeSenderToReceiverContext(), { SenderStreamId });
+				ObjectReplicator->TestValuesWereReplicated(*this);
+			});
+		});
+
+		Describe("When receiver client joins", [this]()
+		{
+			It("Server grants sync control", [this]()
+			{
+				bool bReceivedControlEvent = false;
+				Sender->GetClientSessionMock()->RegisterCustomEventHandler<FConcertReplication_ChangeSyncControl>(
+					[this, &bReceivedControlEvent](const FConcertSessionContext&, const FConcertReplication_ChangeSyncControl& Event)
+					{
+						bReceivedControlEvent = true;
+						
+						TestEqual(TEXT("Update has 1 object"), Event.NewControlStates.Num(), 1);
+						const bool* bMayReplicate = Event.NewControlStates.Find({ SenderStreamId, ObjectReplicator->TestObject });
+						TestTrue(TEXT("Is allowed to replicate"), bMayReplicate && *bMayReplicate);
+					});
+				
+				// This will cause the server to send a sync control event
+				Receiver->JoinReplicationAsListener({ ObjectReplicator->TestObject });
+
+				TestTrue(TEXT("Server sent sync control event"), bReceivedControlEvent);
+			});
+		});
+		
+		Describe("If server revokes sync control", [this]()
+		{
+			It("Client will not attempt to replicate", [this]()
+			{
+				// Like above, this will cause the server to send a sync control event...
+				Receiver->JoinReplicationAsListener({ ObjectReplicator->TestObject });
+				
+				// ... but we'll just simulate that the server revoked control for any reason
+				FConcertReplication_ChangeSyncControl RevokeControl;
+				RevokeControl.NewControlStates.Add(FConcertObjectInStreamID{ SenderStreamId, ObjectReplicator->TestObject }, false);
+				Server->GetServerSessionMock()->SendCustomEvent(RevokeControl, Sender->GetClientSessionMock()->GetSessionClientEndpointId(), EConcertMessageFlags::ReliableOrdered);
+				
+				ObjectReplicator->SimulateSendObjectToReceiver(*this, MakeSenderToReceiverContext(), { SenderStreamId },
+					[this](const FConcertSessionContext& Context, const FConcertReplication_BatchReplicationEvent& Event)
+					{
+						AddError(TEXT("Server: Sending client replicated object even though they did not have permission."));
+					},
+					[this](const FConcertSessionContext& Context, const FConcertReplication_BatchReplicationEvent& Event)
+					{
+						// This error should not fire without the above error (if it does there's a bug in our testing framework)
+						AddError(TEXT("Receiver: Sending client replicated object even though they did not have permission."));
+					});
+				ObjectReplicator->TestValuesWereNotReplicated(*this);
+			});
+		});
+
+		Describe("When receiver client leaves", [this]()
+		{
+			// Joining the 2nd client will cause the server to give the sender client sync control
+			BeforeEach([this]() { Receiver->JoinReplicationAsListener({ ObjectReplicator->TestObject }); });
+			
+			It("Client will stop replicating", [this]()
+			{
+				// Sanity test that sending works
+				ObjectReplicator->SimulateSendObjectToReceiver(*this, MakeSenderToReceiverContext(), { SenderStreamId });
+				ObjectReplicator->TestValuesWereReplicated(*this);
+
+				Receiver->GetClientReplicationManager().LeaveReplicationSession();
+				ObjectReplicator->SimulateSendObjectToReceiver(*this, MakeSenderToReceiverContext(), { SenderStreamId });
+				ObjectReplicator->TestValuesWereNotReplicated(*this);
+			});
+
+			It("Server revokes sync control from 1st client", [this]()
+			{
+				bool bReceivedControlEvent = false;
+				Sender->GetClientSessionMock()->RegisterCustomEventHandler<FConcertReplication_ChangeSyncControl>(
+					[this, &bReceivedControlEvent](const FConcertSessionContext&, const FConcertReplication_ChangeSyncControl& Event)
+					{
+						bReceivedControlEvent = true;
+						
+						TestEqual(TEXT("Update has 1 object"), Event.NewControlStates.Num(), 1);
+						const bool* bMayReplicate = Event.NewControlStates.Find({ SenderStreamId, ObjectReplicator->TestObject });
+						TestTrue(TEXT("Is not allowed to replicate"), bMayReplicate && !*bMayReplicate);
+					});
+
+				// And now we check that the 2nd sync control event is received.
+				Receiver->GetClientReplicationManager().LeaveReplicationSession();
+				TestTrue(TEXT("Server sent sync control event"), bReceivedControlEvent);
+			});
+		});
+		
+		Describe("When there are 3 clients",[this]()
+		{
+			It("Sync control is granted only on first join and revoked only on last leave.", [this]()
+			{
+				FReplicationClient& ThirdClient = Server->ConnectClient();
+
+				bool bReceivedControlEvent = false;
+				Sender->GetClientSessionMock()->RegisterCustomEventHandler<FConcertReplication_ChangeSyncControl>(
+					[this, &bReceivedControlEvent](const FConcertSessionContext&, const FConcertReplication_ChangeSyncControl& Event)
+					{
+						bReceivedControlEvent = true;
+					});
+
+				// 2nd client joining sends the event
+				Receiver->JoinReplicationAsListener({ ObjectReplicator->TestObject });
+				TestTrue(TEXT("2nd client joining causes sync control to be sent"), bReceivedControlEvent);
+
+				// 3rd client joining does not send the event because the sender is now sending to clients 2 & 3
+				bReceivedControlEvent = false;
+				ThirdClient.JoinReplicationAsListener({ ObjectReplicator->TestObject });
+				TestFalse(TEXT("3rd receiver joining does not cause any sync control to be sent"), bReceivedControlEvent);
+
+				// 2nd client leaving does not send because the sender is still allowed to send to client 3
+				bReceivedControlEvent = false;
+				Receiver->GetClientReplicationManager().LeaveReplicationSession();
+				TestFalse(TEXT("2nd receiver leaving does not cause any sync control to be sent"), bReceivedControlEvent);
+
+				// 3rd client leaving does cause the event to be sent
+				bReceivedControlEvent = false;
+				ThirdClient.GetClientReplicationManager().LeaveReplicationSession();
+				TestTrue(TEXT("3rd receiver causes a sync control to be sent"), bReceivedControlEvent);
+			});
+
+			It("Replication continues to work when only 1 client leaves", [this]()
+			{
+				FReplicationClient& ThirdClient = Server->ConnectClient();
+				
+				Receiver->JoinReplicationAsListener({ ObjectReplicator->TestObject });
+				ThirdClient.JoinReplicationAsListener({ ObjectReplicator->TestObject });
+
+				// This should not send any sync control event because the 3rd client is still there (tested in previous case) ...
+				Receiver->GetClientReplicationManager().LeaveReplicationSession();
+
+				// ... hence replication still works, too.
+				ObjectReplicator->SimulateSendObjectToReceiver(*this, MakeSenderToReceiverContext(ThirdClient), { SenderStreamId });
+				ObjectReplicator->TestValuesWereReplicated(*this);
+			});
+		});
+
+		Describe("When sender changes authority", [this]()
+		{
+			BeforeEach([this]()
+			{
+				// Joining the 2nd client will cause the server to give the sender client sync control
+				Receiver->JoinReplicationAsListener({ ObjectReplicator->TestObject });
+
+				// No events are sent for authority changes, see FConcertReplication_ChangeSyncControl documentation.
+				Sender->GetClientSessionMock()->RegisterCustomEventHandler<FConcertReplication_ChangeSyncControl>(
+					[this](const FConcertSessionContext&, const FConcertReplication_ChangeSyncControl& Event)
+					{
+						AddError(TEXT("No sync event was supposed to be sent."));
+					});
+			});
+			
+			It("Client loses sync control when giving up authority explicitly", [this]()
+			{
+				Sender->GetClientReplicationManager().ReleaseAuthorityOf({ ObjectReplicator->TestObject })
+					.Next([this](FConcertReplication_ChangeAuthority_Response&& Response)
+					{
+						// Server did not include the sync loss because the change is implicit (documentation).
+						TestTrue(TEXT("SyncControl.IsEmpty()"), Response.SyncControl.IsEmpty());
+					});
+				
+				// This tests that the sender no longer has sync control
+				ObjectReplicator->SimulateSendObjectToReceiver(*this, MakeSenderToReceiverContext(), { SenderStreamId });
+				ObjectReplicator->TestValuesWereNotReplicated(*this);
+			});
+
+			It("Client loses sync control when giving up authority implicitly", [this]()
+			{
+				FConcertReplication_ChangeStream_Request StreamChange;
+				StreamChange.ObjectsToRemove.Add({ SenderStreamId, ObjectReplicator->TestObject });
+				Sender->GetClientReplicationManager().ChangeStream(StreamChange);
+				
+				// This tests that the sender no longer has sync control
+				ObjectReplicator->SimulateSendObjectToReceiver(*this, MakeSenderToReceiverContext(), { SenderStreamId });
+				ObjectReplicator->TestValuesWereNotReplicated(*this);
+			});
+
+			It("Client gets sync control again when taking back authority", [this]()
+			{
+				Sender->GetClientReplicationManager().ReleaseAuthorityOf({ ObjectReplicator->TestObject });
+				Sender->GetClientReplicationManager().TakeAuthorityOver({ ObjectReplicator->TestObject })
+					.Next([this](FConcertReplication_ChangeAuthority_Response&& Response)
+					{
+						// The server should instantly give sync control because there is another client
+						TestEqual(TEXT("SyncControl.Num() == 1"), Response.SyncControl.NewControlStates.Num(), 1);
+						
+						const FConcertObjectInStreamID ObjectId { SenderStreamId, ObjectReplicator->TestObject };
+						const bool* bNewState =Response.SyncControl.NewControlStates.Find(ObjectId);
+						TestTrue(TEXT("TestObject is enabled again"), bNewState && *bNewState);
+					});
+			});
+		});
+
+		// At this point, everything should work server-side. Now test client-side prediction.
+		Describe("Local client sync control cache", [this]()
+		{
+			// Joining the 2nd client will cause the server to give the sender client sync control
+			BeforeEach([this]()
+			{
+				Receiver->JoinReplicationAsListener({ ObjectReplicator->TestObject });
+
+				// No events are sent for joins nor authority changes, see FConcertReplication_ChangeSyncControl documentation.
+				Sender->GetClientSessionMock()->RegisterCustomEventHandler<FConcertReplication_ChangeSyncControl>(
+					[this](const FConcertSessionContext&, const FConcertReplication_ChangeSyncControl& Event)
+					{
+						AddError(TEXT("Server sent sync control event."));
+					});
+			});
+			
+			It("Is correct after join", [this]
+			{
+				IConcertClientReplicationManager& Manager = Sender->GetClientReplicationManager();
+				TestEqual("NumSyncControlledObjects() == 1", Manager.NumSyncControlledObjects(), 1);
+				TestTrue("HasSyncControl()", Manager.HasSyncControl({ SenderStreamId, ObjectReplicator->TestObject }));
+			});
+
+			It("Is correct after explicit authority change", [this]()
+			{
+				IConcertClientReplicationManager& Manager = Sender->GetClientReplicationManager();
+				
+				// First release authority over the object and test that the client's local sync control cache updates correctly.
+				ConcertSyncCore::Replication::FSyncControlState SyncControlBeforeRelease = Manager.GetSyncControlledObjects();
+				FConcertReplication_ChangeAuthority_Request ReleaseAuthority { .ReleaseAuthority = { { ObjectReplicator->TestObject, {{ SenderStreamId }} } } };
+				Manager.RequestAuthorityChange(ReleaseAuthority)
+					.Next([this, &Manager, &ReleaseAuthority, &SyncControlBeforeRelease](FConcertReplication_ChangeAuthority_Response&& Response)
+					{
+						// Aggregate is validated in SyncControlState.spec.cpp
+						SyncControlBeforeRelease.AppendChanges(ReleaseAuthority, Response);
+						TestTrue(TEXT("Client predicted correctly"), SyncControlBeforeRelease == Manager.GetSyncControlledObjects());
+					});
+
+				// Now take back authority and test that the client's local sync control cache updates correctly.
+				ConcertSyncCore::Replication::FSyncControlState SyncControlBeforeTake = Manager.GetSyncControlledObjects();
+				FConcertReplication_ChangeAuthority_Request TakeAuthority { .TakeAuthority = { { ObjectReplicator->TestObject, {{ SenderStreamId }} } } };
+				Manager.RequestAuthorityChange(ReleaseAuthority)
+					.Next([this, &Manager, &SyncControlBeforeTake, &TakeAuthority](FConcertReplication_ChangeAuthority_Response&& Response)
+					{
+						// Aggregate is validated in SyncControlState.spec.cpp
+						SyncControlBeforeTake.AppendChanges(TakeAuthority, Response);
+						TestTrue(TEXT("Client predicted correctly"), SyncControlBeforeTake == Manager.GetSyncControlledObjects());
+					});
+			});
+
+			It("Is correct after implicit authority change", [this]()
+			{
+				IConcertClientReplicationManager& Manager = Sender->GetClientReplicationManager();
+				
+				// Removing the object will implicitly get rid of the object's authority and sync control
+				ConcertSyncCore::Replication::FSyncControlState SyncControlBeforeRelease = Manager.GetSyncControlledObjects();
+				FConcertReplication_ChangeStream_Request StreamChange;
+				StreamChange.ObjectsToRemove.Add({ SenderStreamId, ObjectReplicator->TestObject });
+				Manager.ChangeStream(StreamChange)
+					.Next([this, &Manager, &StreamChange, &SyncControlBeforeRelease](FConcertReplication_ChangeStream_Response&& Response)
+					{
+						// Aggregate is validated in SyncControlState.spec.cpp
+						SyncControlBeforeRelease.AppendChanges(StreamChange);
+						TestTrue(TEXT("Client predicted correctly"), SyncControlBeforeRelease == Manager.GetSyncControlledObjects());
+					});
+			});
+		});
+	}
+}
diff --git a/Engine/Plugins/Developer/Concert/ConcertSync/ConcertSyncTest/Source/ConcertSyncTest/Private/Replication/Util/SendReceiveObjectTestBase.h b/Engine/Plugins/Developer/Concert/ConcertSync/ConcertSyncTest/Source/ConcertSyncTest/Private/Replication/Util/SendReceiveObjectTestBase.h
--- a/Engine/Plugins/Developer/Concert/ConcertSync/ConcertSyncTest/Source/ConcertSyncTest/Private/Replication/Util/SendReceiveObjectTestBase.h
+++ b/Engine/Plugins/Developer/Concert/ConcertSync/ConcertSyncTest/Source/ConcertSyncTest/Private/Replication/Util/SendReceiveObjectTestBase.h
@@ -60,7 +60,6 @@ namespace UE::ConcertSyncTests::Replication
 		//~ End FSendReceiveTestBase Interface
 		
 		/**
-		 * bHasServerReceivedData and bHasClientReceivedData are reset to false prior to sending.
 		 * Sets test values on TestObject and sends it to the receiver.
 		 * If the data arrived, bHasServerReceivedData and bHasClientReceivedData are true.
 		 */
diff --git a/Engine/Plugins/Developer/Concert/ConcertSync/ConcertSyncTest/Source/ConcertSyncTest/Private/Replication/Util/Spec/ObjectTestReplicator.cpp b/Engine/Plugins/Developer/Concert/ConcertSync/ConcertSyncTest/Source/ConcertSyncTest/Private/Replication/Util/Spec/ObjectTestReplicator.cpp
new file mode 100644
--- /dev/null
+++ b/Engine/Plugins/Developer/Concert/ConcertSync/ConcertSyncTest/Source/ConcertSyncTest/Private/Replication/Util/Spec/ObjectTestReplicator.cpp
@@ -0,0 +1,148 @@
+﻿// Copyright Epic Games, Inc. All Rights Reserved.
+
+#include "ObjectTestReplicator.h"
+
+#include "ReplicationClient.h"
+#include "ReplicationServer.h"
+#include "Replication/PropertyChainUtils.h"
+#include "Replication/Messages/ObjectReplication.h"
+
+namespace UE::ConcertSyncTests::Replication
+{
+	ConcertSyncClient::Replication::FJoinReplicatedSessionArgs FObjectTestReplicator::CreateSenderArgs(
+		FGuid SenderStreamId,
+		EConcertObjectReplicationMode ReplicationMode,
+		uint8 ReplicationRate
+		) const
+	{
+		ConcertSyncClient::Replication::FJoinReplicatedSessionArgs SenderJoinArgs;
+		SenderJoinArgs.Streams.Add(CreateStream(SenderStreamId, ReplicationMode, ReplicationRate));
+		return SenderJoinArgs;
+	}
+
+	FConcertReplicationStream FObjectTestReplicator::CreateStream(FGuid SenderStreamId, EConcertObjectReplicationMode ReplicationMode, uint8 ReplicationRate) const
+	{
+		FConcertReplicatedObjectInfo ReplicatedObjectInfo { TestObject->GetClass() };
+		ConcertSyncCore::PropertyChain::ForEachReplicatableConcertProperty(*TestObject->GetClass(), [&ReplicatedObjectInfo](FConcertPropertyChain&& Chain)
+		{
+			ReplicatedObjectInfo.PropertySelection.ReplicatedProperties.Emplace(MoveTemp(Chain));
+			return EBreakBehavior::Continue;
+		});
+
+		FConcertReplicationStream SendingStream;
+		SendingStream.BaseDescription.Identifier = SenderStreamId;
+		SendingStream.BaseDescription.ReplicationMap.ReplicatedObjects.Add(TestObject, ReplicatedObjectInfo);
+		SendingStream.BaseDescription.FrequencySettings.Defaults = { ReplicationMode, ReplicationRate };
+		return SendingStream;
+	}
+
+	void FObjectTestReplicator::SimulateSendObjectToReceiver(
+		FAutomationTestBase& Test,
+		FObjectReplicationContext Context,
+		TConstArrayView<FGuid> SenderStreams,
+		TFunctionRef<FReceiveReplicationEventSignature> OnServerReceive,
+		TFunctionRef<FReceiveReplicationEventSignature> OnReceiverClientReceive,
+		EPropertyReplicationFlags PropertyFlags
+		) const
+	{
+		auto TestReplicationData_Server = [this, &Test, &SenderStreams, &OnServerReceive](const FConcertSessionContext& Context, const FConcertReplication_BatchReplicationEvent& Event)
+		{
+			Test.TestEqual(TEXT("Server received right number of streams"), Event.Streams.Num(), SenderStreams.Num());
+			for (int32 i = 0; i < Event.Streams.Num(); ++i)
+			{
+				Test.TestEqual(TEXT("Server received 1 object"), Event.Streams[i].ReplicatedObjects.Num() , 1);
+				Test.TestTrue(TEXT("Server received from correct stream"), SenderStreams.Contains(Event.Streams[i].StreamId));
+				const FSoftObjectPath ObjectPath = Event.Streams[i].ReplicatedObjects.IsEmpty() ? FSoftObjectPath{} : Event.Streams[i].ReplicatedObjects[0].ReplicatedObject;
+				Test.TestEqual(TEXT("Server's received object has correct path"), ObjectPath, FSoftObjectPath(TestObject));
+			}
+			OnServerReceive(Context, Event);
+		};
+		auto TestReplicationData_Client_Receiver = [this, &Test, &SenderStreams, &OnReceiverClientReceive](const FConcertSessionContext& Context, const FConcertReplication_BatchReplicationEvent& Event)
+		{
+			Test.TestEqual(TEXT("Client 2 received right number of streams"), Event.Streams.Num(), SenderStreams.Num());
+			for (int32 i = 0; i < Event.Streams.Num(); ++i)
+			{
+				Test.TestEqual(TEXT("Client 2 received 1 object"), Event.Streams[i].ReplicatedObjects.Num() , 1);
+				Test.TestTrue(TEXT("Client 2 received from correct stream"), SenderStreams.Contains(Event.Streams[i].StreamId));
+				const FSoftObjectPath ObjectPath = Event.Streams[i].ReplicatedObjects.IsEmpty() ? FSoftObjectPath{} : Event.Streams[i].ReplicatedObjects[0].ReplicatedObject;
+				Test.TestEqual(TEXT("Client 2's received object has correct path"), ObjectPath, FSoftObjectPath(TestObject));
+			}
+			OnReceiverClientReceive(Context, Event);
+		};
+		const FDelegateHandle ServerHandle = Context.Server.GetServerSessionMock()->RegisterCustomEventHandler<FConcertReplication_BatchReplicationEvent>(TestReplicationData_Server);
+		const FDelegateHandle ClientHandle = Context.Receiver.GetClientSessionMock()->RegisterCustomEventHandler<FConcertReplication_BatchReplicationEvent>(TestReplicationData_Client_Receiver);
+
+		
+		// TestObject is the same UObject on both clients.
+		// Hence we must override test values with SetTestValues and SetDifferentValues.
+		// 1 Sender > Server
+		SetTestValues(PropertyFlags);
+		Context.Sender.TickClient();
+		
+		// 2 Forward from server to receiver
+		Context.Server.TickServer();
+		
+		// 3 Receive from server
+		SetDifferentValues(PropertyFlags);
+		Context.Receiver.TickClient();
+
+		Context.Server.GetServerSessionMock()->UnregisterCustomEventHandler<FConcertReplication_BatchReplicationEvent>(ServerHandle);
+		Context.Receiver.GetClientSessionMock()->UnregisterCustomEventHandler<FConcertReplication_BatchReplicationEvent>(ClientHandle);
+	}
+
+	void FObjectTestReplicator::SetTestValues(EPropertyReplicationFlags PropertyFlags) const
+	{
+		const bool bSendCDOValues = EnumHasAnyFlags(PropertyFlags, EPropertyReplicationFlags::SendCDOValues);
+		if (EnumHasAnyFlags(PropertyFlags, EPropertyReplicationFlags::Float))
+		{
+			TestObject->Float = bSendCDOValues ? GetMutableDefault<UTestReflectionObject>()->Float : SentFloat;
+		}
+		if (EnumHasAnyFlags(PropertyFlags, EPropertyReplicationFlags::Vector))
+		{
+			TestObject->Vector = bSendCDOValues ? GetMutableDefault<UTestReflectionObject>()->Vector : SentVector;
+		}
+	}
+
+	void FObjectTestReplicator::SetDifferentValues(EPropertyReplicationFlags PropertyFlags) const
+	{
+		if (EnumHasAnyFlags(PropertyFlags, EPropertyReplicationFlags::Float))
+		{
+			TestObject->Float = DifferentFloat;
+		}
+		if (EnumHasAnyFlags(PropertyFlags, EPropertyReplicationFlags::Vector))
+		{
+			TestObject->Vector = DifferentVector;
+		}
+	}
+
+	void FObjectTestReplicator::TestValuesWereReplicated(FAutomationTestBase& Test, EPropertyReplicationFlags PropertyFlags) const
+	{
+		const bool bSendCDOValues = EnumHasAnyFlags(PropertyFlags, EPropertyReplicationFlags::SendCDOValues);
+		if (EnumHasAnyFlags(PropertyFlags, EPropertyReplicationFlags::Float))
+		{
+			const float ExpectedValue = bSendCDOValues ? GetMutableDefault<UTestReflectionObject>()->Float : SentFloat;
+			const bool bWasReplicated = Test.TestEqual(TEXT("Float"), TestObject->Float, ExpectedValue);
+			Test.AddErrorIfFalse(bWasReplicated, TEXT("Failed to replicate \"Float\" property"));
+		}
+		if (EnumHasAnyFlags(PropertyFlags, EPropertyReplicationFlags::Vector))
+		{
+			const FVector ExpectedValue = bSendCDOValues ? GetMutableDefault<UTestReflectionObject>()->Vector : SentVector;
+			const bool bWasReplicated = Test.TestEqual(TEXT("Vector"), TestObject->Vector, ExpectedValue);
+			Test.AddErrorIfFalse(bWasReplicated, TEXT("Failed to replicate \"Vector\" property"));
+		}
+	}
+
+	void FObjectTestReplicator::TestValuesWereNotReplicated(FAutomationTestBase& Test, EPropertyReplicationFlags PropertyFlags) const
+	{
+		if (EnumHasAnyFlags(PropertyFlags, EPropertyReplicationFlags::Float))
+		{
+			const bool bWasNotReplicated = Test.TestEqual(TEXT("Float"), TestObject->Float, DifferentFloat);
+			Test.AddErrorIfFalse(bWasNotReplicated, TEXT("Probably \"Float\" property was replicated even though it was not supposed to be!"));
+		}
+		if (EnumHasAnyFlags(PropertyFlags, EPropertyReplicationFlags::Vector))
+		{
+			const bool bWasNotReplicated = Test.TestEqual(TEXT("Vector"), TestObject->Vector, DifferentVector);
+			Test.AddErrorIfFalse(bWasNotReplicated, TEXT("Probably \"Vector\" property was replicated even though it was not supposed to be!"));
+		}
+	}
+}
diff --git a/Engine/Plugins/Developer/Concert/ConcertSync/ConcertSyncTest/Source/ConcertSyncTest/Private/Replication/Util/Spec/ObjectTestReplicator.h b/Engine/Plugins/Developer/Concert/ConcertSync/ConcertSyncTest/Source/ConcertSyncTest/Private/Replication/Util/Spec/ObjectTestReplicator.h
new file mode 100644
--- /dev/null
+++ b/Engine/Plugins/Developer/Concert/ConcertSync/ConcertSyncTest/Source/ConcertSyncTest/Private/Replication/Util/Spec/ObjectTestReplicator.h
@@ -0,0 +1,96 @@
+﻿// Copyright Epic Games, Inc. All Rights Reserved.
+
+#pragma once
+
+#include "Replication/IConcertClientReplicationManager.h"
+#include "Replication/TestReflectionObject.h"
+
+#include "Templates/UnrealTemplate.h"
+#include "UObject/UObjectGlobals.h"
+#include "UObject/Package.h"
+
+struct FConcertReplication_BatchReplicationEvent;
+struct FConcertSessionContext;
+
+namespace UE::ConcertSyncTests::Replication
+{
+	class FReplicationServer;
+	class FReplicationClient;
+
+	/** Defines the properties to be tested in FSendReceiveObjectTestBase: some tests may want more control of which properties are sent. */
+	enum class EPropertyReplicationFlags : uint8
+	{
+		None,
+		
+		Float = 1 << 0,
+		Vector = 1 << 1,
+
+		/** The CDO values should be sent. */
+		SendCDOValues = 1 << 2,
+		
+		All = Float | Vector
+	};
+	ENUM_CLASS_FLAGS(EPropertyReplicationFlags);
+
+	struct FObjectReplicationContext
+	{
+		FReplicationClient& Sender;
+		FReplicationServer& Server;
+		FReplicationClient& Receiver;
+	};
+	
+	/**
+	 * This utility is used in the context of 1 FReplicationServer and 2 FReplicationClients and helps testing transmission of properties from UTestReflectionObject.
+	 *
+	 * It assumes you set up
+	 */
+	class FObjectTestReplicator : public FNoncopyable
+	{
+	public:
+		
+		using FReceiveReplicationEventSignature = void(const FConcertSessionContext& Context, const FConcertReplication_BatchReplicationEvent& Event);
+
+		// Test values
+		float SentFloat = 42.f;
+		FVector SentVector = { 21.f, 84.f, -1.f };
+		float DifferentFloat = -420.f;
+		FVector DifferentVector = { -210.f, -840.f, 10.f };
+
+		/** The object that will be transmitted. */
+		UTestReflectionObject* TestObject = NewObject<UTestReflectionObject>(GetTransientPackage());
+
+		/** Util for creating join args for replicating TestObject. */
+		ConcertSyncClient::Replication::FJoinReplicatedSessionArgs CreateSenderArgs(
+			FGuid SenderStreamId = FGuid::NewGuid(),
+			EConcertObjectReplicationMode ReplicationMode = EConcertObjectReplicationMode::Realtime,
+			uint8 ReplicationRate = 30
+			) const;
+		/** Util for creating a stream that replicates TestObject. */
+		FConcertReplicationStream CreateStream(
+			FGuid SenderStreamId = FGuid::NewGuid(),
+			EConcertObjectReplicationMode ReplicationMode = EConcertObjectReplicationMode::Realtime,
+			uint8 ReplicationRate = 30
+		) const;
+		
+		/** This overload allows you to send properties but from multiple streams. */
+		void SimulateSendObjectToReceiver(
+			FAutomationTestBase& Test,
+			FObjectReplicationContext Context,
+			TConstArrayView<FGuid> SenderStreams,
+			TFunctionRef<FReceiveReplicationEventSignature> OnServerReceive = [](auto&, auto&){},
+			TFunctionRef<FReceiveReplicationEventSignature> OnReceiverClientReceive = [](auto&, auto&){},
+			EPropertyReplicationFlags PropertyFlags = EPropertyReplicationFlags::All
+			) const;
+
+		/** Sets the specified properties to its test values */
+		void SetTestValues(EPropertyReplicationFlags PropertyFlags = EPropertyReplicationFlags::All) const;
+		/** Sets the specified properties to values different from the test values */
+		void SetDifferentValues(EPropertyReplicationFlags PropertyFlags = EPropertyReplicationFlags::All) const;
+
+		/** Tests that the specified properties are equal to their test values */
+		void TestValuesWereReplicated(FAutomationTestBase& Test, EPropertyReplicationFlags PropertyFlags = EPropertyReplicationFlags::All) const;
+		/** Tests that the specified properties are equal to the values different from the test values (i.e. the values SetDifferentValues sets). */
+		void TestValuesWereNotReplicated(FAutomationTestBase& Test, EPropertyReplicationFlags PropertyFlags = EPropertyReplicationFlags::All) const;
+	};
+}
+
diff --git a/Engine/Plugins/Developer/Concert/ConcertSync/ConcertSyncTest/Source/ConcertSyncTest/Private/Replication/Util/Spec/ReplicationClient.cpp b/Engine/Plugins/Developer/Concert/ConcertSync/ConcertSyncTest/Source/ConcertSyncTest/Private/Replication/Util/Spec/ReplicationClient.cpp
new file mode 100644
--- /dev/null
+++ b/Engine/Plugins/Developer/Concert/ConcertSync/ConcertSyncTest/Source/ConcertSyncTest/Private/Replication/Util/Spec/ReplicationClient.cpp
@@ -0,0 +1,71 @@
+﻿// Copyright Epic Games, Inc. All Rights Reserved.
+
+#include "ReplicationClient.h"
+
+#include "Replication/IConcertClientReplicationManager.h"
+#include "Replication/ReplicationTestInterface.h"
+#include "Replication/Util/ConcertClientReplicationBridgeMock.h"
+
+#include "UObject/Object.h"
+
+namespace UE::ConcertSyncTests::Replication
+{
+	void FReplicationClient::TickClient(float FakeDeltaTime)
+	{
+		ClientSessionMock->OnTick().Broadcast(*ClientSessionMock, FakeDeltaTime);
+	}
+
+	TFuture<ConcertSyncClient::Replication::FJoinReplicatedSessionResult> FReplicationClient::JoinReplication(
+		ConcertSyncClient::Replication::FJoinReplicatedSessionArgs Args,
+		EReplicationClientFlags Flags
+		)
+	{
+		if (EnumHasAnyFlags(Flags, EReplicationClientFlags::UseRealReplicationBridge))
+		{
+			BridgeUsed = ConcertSyncClient::TestInterface::CreateClientReplicationBridge();
+		}
+		else
+		{
+			BridgeMock =  MakeShared<FConcertClientReplicationBridgeMock>();
+			BridgeUsed = BridgeMock;
+		}
+		ClientReplicationManager = ConcertSyncClient::TestInterface::CreateClientReplicationManager(ClientSessionMock, *BridgeUsed);
+		
+		// 1.1 Sender offers to send all UTestReflectionObject properties
+		bool bSuccess = false;
+		TFuture<ConcertSyncClient::Replication::FJoinReplicatedSessionResult> Future = ClientReplicationManager->JoinReplicationSession(Args)
+			.Next([&bSuccess](ConcertSyncClient::Replication::FJoinReplicatedSessionResult&& Result)
+			{
+				bSuccess = Result.ErrorCode == EJoinReplicationErrorCode::Success;
+				return Result;
+			});
+
+		if (!EnumHasAnyFlags(Flags, EReplicationClientFlags::SkipJoinTest))
+		{
+			TestContext.TestTrue(TEXT("Replication joined successfully"), bSuccess);
+		}
+		
+		return Future;
+	}
+
+	TFuture<ConcertSyncClient::Replication::FJoinReplicatedSessionResult> FReplicationClient::JoinReplicationAsListener(
+		TArray<UObject*> ObjectsToReceive
+		)
+	{
+		return JoinReplication()
+			.Next([this, ObjectsToReceive = MoveTemp(ObjectsToReceive)](ConcertSyncClient::Replication::FJoinReplicatedSessionResult&& Result)
+			{
+				if (Result.ErrorCode != EJoinReplicationErrorCode::Success)
+				{
+					return Result;
+				}
+				
+				for (UObject* Object : ObjectsToReceive)
+				{
+					BridgeMock->InjectAvailableObject(*Object);
+				}
+				return Result;
+			});
+	}
+}
+
diff --git a/Engine/Plugins/Developer/Concert/ConcertSync/ConcertSyncTest/Source/ConcertSyncTest/Private/Replication/Util/Spec/ReplicationClient.h b/Engine/Plugins/Developer/Concert/ConcertSync/ConcertSyncTest/Source/ConcertSyncTest/Private/Replication/Util/Spec/ReplicationClient.h
new file mode 100644
--- /dev/null
+++ b/Engine/Plugins/Developer/Concert/ConcertSync/ConcertSyncTest/Source/ConcertSyncTest/Private/Replication/Util/Spec/ReplicationClient.h
@@ -0,0 +1,73 @@
+﻿// Copyright Epic Games, Inc. All Rights Reserved.
+
+#pragma once
+
+#include "Replication/IConcertClientReplicationManager.h"
+#include "Util/ClientServerCommunicationTest.h"
+#include "Templates/UnrealTemplate.h"
+
+class IConcertClientReplicationManager;
+class IConcertClientReplicationBridge;
+class UObject;
+
+namespace UE::ConcertSyncTests::Replication
+{
+	class FConcertClientReplicationBridgeMock;
+
+	enum class EReplicationClientFlags : uint8
+	{
+		None = 0,
+		/**
+		 * Instead of mocking the bridge use the real implementation. The object will only be discovered if it exists when the stream is registered.
+		 * See IConcertClientReplicationBridge::PushTrackedObjects.
+		 */
+		UseRealReplicationBridge = 1 << 0,
+
+		/** Don't test whether joining was successful */
+		SkipJoinTest = 1 << 1
+	};
+	ENUM_CLASS_FLAGS(EReplicationClientFlags);
+	
+	/** Implements reusable behaviour for replication tests, such as the handshake. */
+	class FReplicationClient : public FNoncopyable
+	{
+	public:
+		
+		FReplicationClient(const FGuid& ClientEndPointId, FConcertServerSessionMock& Server, FAutomationTestBase& TestContext)
+			: TestContext(TestContext)
+			, ClientSessionMock(MakeShared<FConcertClientSessionMock>(ClientEndPointId, Server))
+		{}
+		
+		/** Lets the client process any messages that have come in. */
+		void TickClient(float FakeDeltaTime = 1.f / 60.f);
+
+		/** Joins the client into replication allowing them to participate in replication */
+		TFuture<ConcertSyncClient::Replication::FJoinReplicatedSessionResult> JoinReplication(
+			ConcertSyncClient::Replication::FJoinReplicatedSessionArgs Args = {},
+			EReplicationClientFlags Flags = EReplicationClientFlags::None
+			);
+		/** This overload joins replication and injects the objects into the FConcertClientReplicationBridgeMock so they can be received. */
+		TFuture<ConcertSyncClient::Replication::FJoinReplicatedSessionResult> JoinReplicationAsListener(
+			TArray<UObject*> ObjectsToReceive
+			);
+
+		TSharedRef<FConcertClientSessionBaseMock> GetClientSessionMock() const { return ClientSessionMock; }
+
+		FConcertClientReplicationBridgeMock& GetBridgeMock() const { return *BridgeMock; }
+		IConcertClientReplicationManager& GetClientReplicationManager() const { return *ClientReplicationManager; }
+
+	private:
+
+		/** Used to test "obvious" cases that should never fail in any test. */
+		FAutomationTestBase& TestContext;
+		/** The underlying client session */
+		const TSharedRef<FConcertClientSessionBaseMock> ClientSessionMock;
+		
+		/** Manually detects when objects are "available". This is null if ESendReceiveTestFlags::UseRealReplicationBridge was specified. */
+		TSharedPtr<FConcertClientReplicationBridgeMock> BridgeMock;
+		/** Detects when objects are available. This is always valid. It is the bridge being used by the receiver. */
+		TSharedPtr<IConcertClientReplicationBridge> BridgeUsed;
+		/** Manages replication client side */
+		TSharedPtr<IConcertClientReplicationManager> ClientReplicationManager;
+	};
+}
diff --git a/Engine/Plugins/Developer/Concert/ConcertSync/ConcertSyncTest/Source/ConcertSyncTest/Private/Replication/Util/Spec/ReplicationServer.cpp b/Engine/Plugins/Developer/Concert/ConcertSync/ConcertSyncTest/Source/ConcertSyncTest/Private/Replication/Util/Spec/ReplicationServer.cpp
new file mode 100644
--- /dev/null
+++ b/Engine/Plugins/Developer/Concert/ConcertSync/ConcertSyncTest/Source/ConcertSyncTest/Private/Replication/Util/Spec/ReplicationServer.cpp
@@ -0,0 +1,28 @@
+﻿// Copyright Epic Games, Inc. All Rights Reserved.
+
+#include "ReplicationServer.h"
+
+#include "ReplicationClient.h"
+#include "Replication/ReplicationTestInterface.h"
+
+namespace UE::ConcertSyncTests::Replication
+{
+	FReplicationServer::FReplicationServer(FAutomationTestBase& TestContext)
+		: TestContext(TestContext)
+		, ServerSessionMock(MakeShared<FConcertServerSessionMock>())
+		, ServerReplicationManager(ConcertSyncServer::TestInterface::CreateServerReplicationManager(ServerSessionMock))
+	{}
+
+	FReplicationClient& FReplicationServer::ConnectClient()
+	{
+		const FGuid ClientEndpointId(0, 0, 0, Clients.Num() + 1); // {0, 0, 0, 0} is used by the server.
+		Clients.Add(MakeUnique<FReplicationClient>(ClientEndpointId, *ServerSessionMock, TestContext));
+		ServerSessionMock->ConnectClient(ClientEndpointId, *(Clients.Last()->GetClientSessionMock()));
+		return *Clients.Last();
+	}
+
+	void FReplicationServer::TickServer(float FakeDeltaTime)
+	{
+		ServerSessionMock->OnTick().Broadcast(*ServerSessionMock, FakeDeltaTime);
+	}
+}
diff --git a/Engine/Plugins/Developer/Concert/ConcertSync/ConcertSyncTest/Source/ConcertSyncTest/Private/Replication/Util/Spec/ReplicationServer.h b/Engine/Plugins/Developer/Concert/ConcertSync/ConcertSyncTest/Source/ConcertSyncTest/Private/Replication/Util/Spec/ReplicationServer.h
new file mode 100644
--- /dev/null
+++ b/Engine/Plugins/Developer/Concert/ConcertSync/ConcertSyncTest/Source/ConcertSyncTest/Private/Replication/Util/Spec/ReplicationServer.h
@@ -0,0 +1,45 @@
+﻿// Copyright Epic Games, Inc. All Rights Reserved.
+
+#pragma once
+
+#include "Replication/IConcertServerReplicationManager.h"
+#include "Util/ClientServerCommunicationTest.h"
+#include "Templates/UnrealTemplate.h"
+
+class FAutomationTestBase;
+
+namespace UE::ConcertSyncTests::Replication
+{
+	class FReplicationClient;
+
+	/** Reusable logic to simulate a test server in spec tests with. */
+	class FReplicationServer : public FNoncopyable
+	{
+	public:
+
+		FReplicationServer(FAutomationTestBase& TestContext);
+
+		/** Connects a client to the server. */
+		FReplicationClient& ConnectClient();
+		
+		/** Lets the server process any messages that have come in. */
+		void TickServer(float FakeDeltaTime = 1.f / 60.f);
+
+		const TSharedRef<FConcertServerSessionMock>& GetServerSessionMock() const { return ServerSessionMock; }
+	
+	private:
+
+		/** Used to test "obvious" cases that should never fail in any test. */
+		FAutomationTestBase& TestContext;
+
+		/** The underlying server session */
+		TSharedRef<FConcertServerSessionMock> ServerSessionMock;
+		/** Manages replication server side */
+		TSharedRef<ConcertSyncServer::Replication::IConcertServerReplicationManager> ServerReplicationManager;
+
+		/** Clients connected thus far. */
+		TArray<TUniquePtr<FReplicationClient>> Clients;
+	};
+}
+
+
