Add utilities for soft object paths to Concert:
- ObjectPathUtils: Functions for analysing FSoftObjectPath string (getting outers)
- FObjectPathHierarchy: Builds a hierarchy tree based on FSoftObjectPath that you add and remove. Allows traversing top to bottom and bottom to top.
- FReplicatedObjectHierarchyCache: Keeps track of an object hierarchy for objects registered for replication.

diff --git a/Engine/Plugins/Developer/Concert/ConcertSync/ConcertSyncCore/Source/ConcertSyncCore/Private/Misc/ObjectPathHierarchy.cpp b/Engine/Plugins/Developer/Concert/ConcertSync/ConcertSyncCore/Source/ConcertSyncCore/Private/Misc/ObjectPathHierarchy.cpp
new file mode 100644
--- /dev/null
+++ b/Engine/Plugins/Developer/Concert/ConcertSync/ConcertSyncCore/Source/ConcertSyncCore/Private/Misc/ObjectPathHierarchy.cpp
@@ -0,0 +1,230 @@
+﻿// Copyright Epic Games, Inc. All Rights Reserved.
+
+#include "Misc/ObjectPathHierarchy.h"
+
+#include "Misc/ObjectPathUtils.h"
+
+namespace UE::ConcertSyncCore
+{
+	namespace Private
+	{
+		template<typename T>
+		static int32 IndexOfChildInUniquePtrArray(const TArray<TUniquePtr<T>>& Children, T* Child)
+		{
+			return Children.IndexOfByPredicate([Child](const TUniquePtr<T>& TreeNode)
+			{
+				return TreeNode.Get() == Child;
+			});
+		}
+	}
+
+	void FObjectPathHierarchy::TraverseTopToBottom(TFunctionRef<ETreeTraversalBehavior(const FChildRelation& Relation)> Callback, const FSoftObjectPath& Start) const
+	{
+		const bool bHasStart = !Start.IsNull();
+		if (!bHasStart)
+		{
+			for (const TUniquePtr<FTreeNode>& Asset : AssetNodes)
+			{
+				if (!Asset->Children.IsEmpty()
+					&& TraverseTopToBottomInternal(Asset->Data, Asset->Children, Callback) == ETreeTraversalBehavior::Break)
+				{
+					break;
+				}
+			}
+		}
+		else if (FTreeNode*const*const TreeNode = CachedNodes.Find(Start))
+		{
+			const FTreeNode& Node = **TreeNode;
+			TraverseTopToBottomInternal(Node.Data, Node.Children, Callback);
+		}
+	}
+
+	void FObjectPathHierarchy::TraverseBottomToTop(TFunctionRef<EBreakBehavior(const FChildRelation& Relation)> Callback, const FSoftObjectPath& Start) const
+	{
+		const bool bHasStart = !Start.IsNull();
+		if (!bHasStart)
+		{
+			for (const TUniquePtr<FTreeNode>& Asset : AssetNodes)
+			{
+				if (!Asset->Children.IsEmpty()
+					&& TraverseBottomToTopInternal(Asset->Data, Asset->Children, Callback) == EBreakBehavior::Break)
+				{
+					break;
+				}
+			}
+		}
+		else if (FTreeNode*const*const TreeNode = CachedNodes.Find(Start))
+		{
+			const FTreeNode& Node = **TreeNode;
+			TraverseBottomToTopInternal(Node.Data, Node.Children, Callback);
+		}
+	}
+
+	TOptional<EHierarchyObjectType> FObjectPathHierarchy::IsInHierarchy(const FSoftObjectPath& Object) const
+	{
+		FTreeNode*const*const TreeNode = CachedNodes.Find(Object);
+		return TreeNode
+			? (*TreeNode)->Data.Type
+			: TOptional<EHierarchyObjectType>{};
+	}
+
+	bool FObjectPathHierarchy::HasChildren(const FSoftObjectPath& Object) const
+	{
+		FTreeNode*const*const TreeNode = CachedNodes.Find(Object);
+		return TreeNode && !(*TreeNode)->Children.IsEmpty();
+	}
+
+	void FObjectPathHierarchy::AddObject(const FSoftObjectPath& ObjectPath)
+	{
+		if (FTreeNode** TreeNode = CachedNodes.Find(ObjectPath))
+		{
+			(*TreeNode)->Data.Type = EHierarchyObjectType::Explicit;
+			return;
+		}
+
+		TUniquePtr<FTreeNode> NodeToInsert = MakeUnique<FTreeNode>(FObjectHierarchyInfo{ ObjectPath, EHierarchyObjectType::Explicit });
+		// We'll walk up the hierarchy and in each iteration we insert a child node.
+		for (FSoftObjectPath CurrentPath = ObjectPath; !CurrentPath.IsNull();)
+		{
+			CachedNodes.Add(CurrentPath, NodeToInsert.Get());
+			
+			// Base case: NodeToInsert is supposed to be a root node
+			TOptional<FSoftObjectPath> OptionalParentPath = GetOuterPath(CurrentPath);
+			if (!OptionalParentPath)
+			{
+				AssetNodes.Emplace(MoveTemp(NodeToInsert));
+				return;
+			}
+
+			const FSoftObjectPath& ParentPath = *OptionalParentPath;
+			FTreeNode** ParentNode = CachedNodes.Find(ParentPath);
+			if (ParentNode)
+			{
+				// Just insert the node into the pre-existing hierarchy and we're done.
+				NodeToInsert->Parent = *ParentNode;
+				(*ParentNode)->Children.Emplace(MoveTemp(NodeToInsert));
+				break;
+			}
+			
+			// Parent does not exist, yet, so create it.
+			TUniquePtr<FTreeNode> Parent = MakeUnique<FTreeNode>(FObjectHierarchyInfo{ ParentPath, EHierarchyObjectType::Implicit });
+			NodeToInsert->Parent = Parent.Get();
+			Parent->Children.Emplace(MoveTemp(NodeToInsert));
+
+			// Advance the hierarchy:
+			NodeToInsert = MoveTemp(Parent);
+			CurrentPath = ParentPath;
+		}
+	}
+
+	void FObjectPathHierarchy::RemoveObject(const FSoftObjectPath& Object)
+	{
+		FTreeNode** PossibleTreeNode = CachedNodes.Find(Object);
+		if (!PossibleTreeNode)
+		{
+			return;
+		}
+
+		// This node is required by one of its children
+		FTreeNode& NodeToRemove = **PossibleTreeNode;
+		if (NodeToRemove.Data.Type == EHierarchyObjectType::Implicit)
+		{
+			ensureMsgf(!NodeToRemove.Children.IsEmpty(), TEXT("Broke invariant: there is supposed to be an explicitly added child."));
+			return;
+		}
+
+		// If it has children, by invariant there must still be at least one child that is Explicit
+		if (!NodeToRemove.Children.IsEmpty())
+		{
+			NodeToRemove.Data.Type = EHierarchyObjectType::Implicit;
+			return;
+		}
+
+		// Walk up the hierarchy and destroy all nodes that were implicitly created for NodeToRemove
+		bool bShouldStop = false;
+		for (FTreeNode* CurrentNode = &NodeToRemove; CurrentNode && !bShouldStop;)
+		{
+			// Asset nodes do not have Parent. All other nodes have a valid Parent.
+			if (CurrentNode->IsAssetNode()) 
+			{
+				// Order matters because CurrentNode will be destroyed when removed from AssetNodes.
+				CachedNodes.Remove(CurrentNode->Data.Object);
+				
+				const int32 AssetIndex = Private::IndexOfChildInUniquePtrArray(AssetNodes, CurrentNode);
+				check(AssetNodes.IsValidIndex(AssetIndex));
+				AssetNodes.RemoveAtSwap(AssetIndex);
+				return;
+			}
+			
+			FTreeNode* Parent = CurrentNode->Parent;
+			checkf(Parent, TEXT("CurrentNode %s is not an asset so it should have a parent"), *CurrentNode->Data.Object.ToString());
+			{
+				FTreeNode* NodePendingDestruction = CurrentNode;
+				CachedNodes.Remove(NodePendingDestruction->Data.Object);
+				
+				const int32 ChildIndex = Private::IndexOfChildInUniquePtrArray(Parent->Children, NodePendingDestruction); 
+				if (ensure(Parent->Children.IsValidIndex(ChildIndex)))
+				{
+					Parent->Children.RemoveAtSwap(ChildIndex); // ~TUniquePtr now destroys the node
+				}
+				
+				// Memory pointed to by NodePendingDestruction has been destroyed - do not reference it anymore
+				CurrentNode = Parent;
+			}
+			
+			// No more walking up the chain if we encounter an explicitly added object or a node that another object generated implicitly
+			bShouldStop = CurrentNode->Data.Type == EHierarchyObjectType::Explicit || CurrentNode->Children.Num() > 1;
+		}
+	}
+
+	ETreeTraversalBehavior FObjectPathHierarchy::TraverseTopToBottomInternal(
+		const FObjectHierarchyInfo& OwnerData,
+		const TArray<TUniquePtr<FTreeNode>>& Nodes,
+		TFunctionRef<ETreeTraversalBehavior(const FChildRelation& Relation)> Callback
+		)
+	{
+		for (const TUniquePtr<FTreeNode>& Node : Nodes)
+		{
+			const FObjectHierarchyInfo& ChildData = Node->Data;
+			const FChildRelation Relation { OwnerData, ChildData };
+			switch (Callback(Relation))
+			{
+			case ETreeTraversalBehavior::Continue:
+				if (TraverseTopToBottomInternal(ChildData, Node->Children, Callback) == ETreeTraversalBehavior::Break)
+				{
+					return ETreeTraversalBehavior::Break;
+				}
+				continue;
+			case ETreeTraversalBehavior::SkipSubtree: continue;
+			case ETreeTraversalBehavior::Break: [[fallthrough]];
+			default: return ETreeTraversalBehavior::Break;
+			}
+		}
+		
+		return ETreeTraversalBehavior::Continue;
+	}
+
+	EBreakBehavior FObjectPathHierarchy::TraverseBottomToTopInternal(
+		const FObjectHierarchyInfo& OwnerData,
+		const TArray<TUniquePtr<FTreeNode>>& Nodes,
+		TFunctionRef<EBreakBehavior(const FChildRelation& Relation)> Callback
+		)
+	{
+		for (const TUniquePtr<FTreeNode>& Node : Nodes)
+		{
+			if (TraverseBottomToTopInternal(Node->Data, Node->Children, Callback) == EBreakBehavior::Break)
+			{
+				return EBreakBehavior::Break;
+			}
+			
+			const FObjectHierarchyInfo& ChildData = Node->Data;
+			const FChildRelation Relation { OwnerData, ChildData };
+			if (Callback(Relation) == EBreakBehavior::Break)
+			{
+				return EBreakBehavior::Break;
+			}
+		}
+		
+		return EBreakBehavior::Continue;
+	}
+}
diff --git a/Engine/Plugins/Developer/Concert/ConcertSync/ConcertSyncCore/Source/ConcertSyncCore/Private/Replication/Misc/ReplicatedObjectHierarchyCache.cpp b/Engine/Plugins/Developer/Concert/ConcertSync/ConcertSyncCore/Source/ConcertSyncCore/Private/Replication/Misc/ReplicatedObjectHierarchyCache.cpp
new file mode 100644
--- /dev/null
+++ b/Engine/Plugins/Developer/Concert/ConcertSync/ConcertSyncCore/Source/ConcertSyncCore/Private/Replication/Misc/ReplicatedObjectHierarchyCache.cpp
@@ -0,0 +1,107 @@
+﻿// Copyright Epic Games, Inc. All Rights Reserved.
+
+#include "Replication/Misc/ReplicatedObjectHierarchyCache.h"
+
+#include "Replication/ChangeStreamSharedUtils.h"
+#include "Replication/Messages/Handshake.h"
+
+#include "Algo/AnyOf.h"
+
+namespace UE::ConcertSyncCore
+{
+	void FReplicatedObjectHierarchyCache::Clear()
+	{
+		ObjectMetaData.Reset();
+		FObjectPathHierarchy::Clear();
+	}
+
+	bool FReplicatedObjectHierarchyCache::IsObjectReferencedDirectly(const FSoftObjectPath& ObjectPath, TConstArrayView<FGuid> IgnoredClients) const
+	{
+		const FObjectMetaData* MetaData = ObjectMetaData.Find(ObjectPath);
+		return MetaData && Algo::AnyOf(MetaData->ReferencingStreams, [&IgnoredClients](const FStreamReferencer& Referencer)
+		{
+			const bool bIsIgnored = IgnoredClients.Contains(Referencer.ClientId);
+			return !bIsIgnored;
+		});
+	}
+	
+	void FReplicatedObjectHierarchyCache::OnJoin(const FGuid& ClientId, const FConcertReplication_Join_Request& Request)
+	{
+		for (const FConcertReplicationStream& Stream : Request.Streams)
+		{
+			const FGuid& StreamId = Stream.BaseDescription.Identifier;
+			for (const TPair<FSoftObjectPath, FConcertReplicatedObjectInfo>& ObjectInfo : Stream.BaseDescription.ReplicationMap.ReplicatedObjects)
+			{
+				AddObjectInternal({ { StreamId, ObjectInfo.Key }, ClientId });
+			}
+		}
+	}
+
+	void FReplicatedObjectHierarchyCache::OnPreClientLeft(const FGuid& ClientId, TConstArrayView<FConcertReplicationStream> Streams)
+	{
+		for (const FConcertReplicationStream& Stream : Streams)
+		{
+			const FGuid& StreamId = Stream.BaseDescription.Identifier;
+			for (const TPair<FSoftObjectPath, FConcertReplicatedObjectInfo>& ObjectInfo : Stream.BaseDescription.ReplicationMap.ReplicatedObjects)
+			{
+				RemoveObjectInternal({ { StreamId, ObjectInfo.Key }, ClientId });
+			}
+		}
+	}
+
+	void FReplicatedObjectHierarchyCache::OnChangeStreams(const FGuid& ClientId, TConstArrayView<FConcertObjectInStreamID> AddedObjects, TConstArrayView<FConcertObjectInStreamID> RemovedObjects)
+	{
+		for (const FConcertObjectInStreamID& RemovedObject : RemovedObjects)
+		{
+			RemoveObjectInternal({ { RemovedObject }, ClientId });
+		}
+
+		for (const FConcertObjectInStreamID& AddeddObject : AddedObjects)
+		{
+			AddObjectInternal({ { AddeddObject }, ClientId });
+		}
+	}
+
+	void FReplicatedObjectHierarchyCache::AddObjectInternal(const FConcertReplicatedObjectId& ObjectInfo)
+	{
+		AddObject(ObjectInfo.Object);
+		
+		ObjectMetaData.FindOrAdd(ObjectInfo.Object)
+			.ReferencingStreams.AddUnique({ ObjectInfo.SenderEndpointId, ObjectInfo.StreamId });
+	}
+
+	void FReplicatedObjectHierarchyCache::RemoveObjectInternal(const FConcertReplicatedObjectId& Object)
+	{
+		if (RemoveMetaData(Object))
+		{
+			RemoveObject(Object.Object);
+		}
+	}
+
+	bool FReplicatedObjectHierarchyCache::RemoveMetaData(const FConcertReplicatedObjectId& Object)
+	{
+		FObjectMetaData* MetaData = ObjectMetaData.Find(Object.Object);
+		// Every AddObjectInternal call should be matched with a RemoveObjectInternal call
+		if (!ensure(MetaData))
+		{
+			return false;
+		}
+		
+		const int32 Index = MetaData->ReferencingStreams.IndexOfByPredicate([&Object](const FStreamReferencer& Referencer)
+		{
+			return Referencer.ClientId == Object.SenderEndpointId && Referencer.StreamId == Object.StreamId;
+		});
+		const bool bIsValidIndex = MetaData->ReferencingStreams.IsValidIndex(Index);
+		if (bIsValidIndex && MetaData->ReferencingStreams.Num() == 1)
+		{
+			ObjectMetaData.Remove(Object.Object);
+			return true;
+		}
+
+		if (bIsValidIndex)
+		{
+			MetaData->ReferencingStreams.RemoveAtSwap(Index);
+		}
+		return false;
+	}
+}
diff --git a/Engine/Plugins/Developer/Concert/ConcertSync/ConcertSyncCore/Source/ConcertSyncCore/Public/Misc/ObjectPathHierarchy.h b/Engine/Plugins/Developer/Concert/ConcertSync/ConcertSyncCore/Source/ConcertSyncCore/Public/Misc/ObjectPathHierarchy.h
new file mode 100644
--- /dev/null
+++ b/Engine/Plugins/Developer/Concert/ConcertSync/ConcertSyncCore/Source/ConcertSyncCore/Public/Misc/ObjectPathHierarchy.h
@@ -0,0 +1,167 @@
+﻿// Copyright Epic Games, Inc. All Rights Reserved.
+
+#pragma once
+
+#include "EBreakBehavior.h"
+
+#include "Containers/Array.h"
+#include "Containers/Map.h"
+#include "UObject/SoftObjectPath.h"
+#include "Templates/UniquePtr.h"
+#include "Templates/UnrealTemplate.h"
+
+namespace UE::ConcertSyncCore
+{
+	enum class ETreeTraversalBehavior : uint8
+	{
+		/** Gives us the next element (may be child or neigbhour) */
+		Continue,
+		/** Gives us the next neighbour element - do not list any children of this node. */
+		SkipSubtree,
+		/** Stop iteration */
+		Break
+	};
+
+	enum class EHierarchyObjectType
+	{
+		/**
+		 * The object entry was added through an AddObject call.
+		 * Example: AddObject(/Game/Maps.Map:PersistentLevel.Cube.StaticMeshComponent1)
+		 */
+		Explicit,
+		/**
+		 * The object entry was added indirectly because of a AddObject to a child.
+		 * Example: AddObject(/Game/Maps.Map:PersistentLevel.Cube.StaticMeshComponent1) adds /Game/Maps.Map:PersistentLevel.Cube implicitly.
+		 */
+		Implicit
+	};
+	
+	struct FObjectHierarchyInfo
+	{
+		FSoftObjectPath Object;
+		EHierarchyObjectType Type;
+	};
+	
+	struct FChildRelation
+	{
+		FObjectHierarchyInfo Parent;
+		FObjectHierarchyInfo Child;
+	};
+	
+	/** Keeps track of the outer tree hierarchy of FSoftObjectPath. */
+	class CONCERTSYNCCORE_API FObjectPathHierarchy
+		// This class cannot be shallowly copied because the nodes referenced dynamic allocations - we don't expect instances to be copied, so we won't bother implementing a deep copy.
+		: public FNoncopyable
+	{
+	public:
+		
+		/**
+		 * Traverses the hierarchy in pre-order (root first, then its children), starting at an optional parent object.
+		 *
+		 * Example sequence of callback invocations:
+		 * 1. (/Game/Maps.Map,							/Game/Maps.Map:PersistentLevel)
+		 * 2. (/Game/Maps.Map:PersistentLevel,			/Game/Maps.Map:PersistentLevel.Cube)
+		 * 3. (/Game/Maps.Map:PersistentLevel.Cube,		/Game/Maps.Map:PersistentLevel.Cube.StaticMeshComponent0)
+		 * 4. (/Game/Maps.Map:PersistentLevel.Cube,		/Game/Maps.Map:PersistentLevel.Cube.StaticMeshComponent1)
+		 * 5. (/Game/Maps.Map:PersistentLevel,			/Game/Maps.Map:PersistentLevel.Sphere)
+		 * 6. (/Game/Maps.Map:PersistentLevel.Sphere,	/Game/Maps.Map:PersistentLevel.Sphere.StaticMeshComponent0)
+		 *
+		 * @important This lists implicit parents.
+		 * For example, if you called AddObject(/Game/Maps.Map:PersistentLevel.Cube.StaticMeshComponent0), then its outers will be implicitly
+		 * considered part of the hierarchy (i.e. /Game/Maps.Map:PersistentLevel.Cube, /Game/Maps.Map:PersistentLevel, /Game/Maps.Map).
+		 * Use the FChildRelation to determine whether an object is part of the hiearchy implicitly or explicitly.
+		 *
+		 * Assets without subobjects are not listed.
+		 * For example, if you called AddObject(/Game/Maps.Map) and the hierarchy was otherwise empty, Callback would not be invoked.
+		 * 
+		 * @param Callback The callback to invoke 
+		 * @param Start The object at which to start walking down the hierarchy
+		 */
+		void TraverseTopToBottom(TFunctionRef<ETreeTraversalBehavior(const FChildRelation& Relation)> Callback, const FSoftObjectPath& Start = {}) const;
+
+		/**
+		 * Traverses the hierarchy in post-order (children first, then the root), starting at an optional parent object.
+		 *
+		 * Example sequence of callback invocations:
+		 * 1. (/Game/Maps.Map:PersistentLevel.Cube,		/Game/Maps.Map:PersistentLevel.Cube.StaticMeshComponent0)
+		 * 2. (/Game/Maps.Map:PersistentLevel.Cube,		/Game/Maps.Map:PersistentLevel.Cube.StaticMeshComponent1)
+		 * 3. (/Game/Maps.Map:PersistentLevel,			/Game/Maps.Map:PersistentLevel.Cube)
+		 * 4. (/Game/Maps.Map:PersistentLevel.Sphere,	/Game/Maps.Map:PersistentLevel.Sphere.StaticMeshComponent0)
+		 * 5. (/Game/Maps.Map:PersistentLevel,			/Game/Maps.Map:PersistentLevel.Sphere)
+		 * 6. (/Game/Maps.Map,							/Game/Maps.Map:PersistentLevel)
+		 *
+		 * @important This lists implicit parents.
+		 * For example, if you called AddObject(/Game/Maps.Map:PersistentLevel.Cube.StaticMeshComponent0), then its outers will be implicitly
+		 * considered part of the hierarchy (i.e. /Game/Maps.Map:PersistentLevel.Cube, /Game/Maps.Map:PersistentLevel, /Game/Maps.Map).
+		 * Use the FChildRelation to determine whether an object is part of the hiearchy implicitly or explicitly.
+		 *
+		 * Assets without children are not listed.
+		 * For example, if you called AddObject(/Game/Maps.Map) and the hierarchy was otherwise empty, Callback would not be invoked.
+		 * 
+		 * @param Callback The callback to invoke 
+		 * @param Start The object at which to start walking up the hierarchy
+		 */
+		void TraverseBottomToTop(TFunctionRef<EBreakBehavior(const FChildRelation& Relation)> Callback, const FSoftObjectPath& Start = {}) const;
+
+		/** @return A valid TOptional if the object is in the hierarchy with the value indicating whether implicitly or explicitly. Unset if the object does not appear in the hierarchy.*/
+		TOptional<EHierarchyObjectType> IsInHierarchy(const FSoftObjectPath& Object) const;
+		/** @return Whether Object has any subobjects in the hierarchy. */
+		bool HasChildren(const FSoftObjectPath& Object) const;
+
+		/** Checks there is an hierarchy. */
+		bool IsEmpty() const { return AssetNodes.IsEmpty(); }
+		
+		/**
+		 * Adds Object to the hierarchy. 
+		 * Henceforth, the object shall be tracked as EHierarchyObjectType::Explicit.
+		 * Calling this more than once has no effect.
+		 */
+		void AddObject(const FSoftObjectPath& ObjectPath);
+
+		/**
+		 * Removes Object from the hierarchy.
+		 *
+		 * If the object has no children, the object is removed entirely.
+		 * If the object has children, it shall henceforth be tacked as EHierarchyObjectType::Implicit.
+		 */
+		void RemoveObject(const FSoftObjectPath& Object);
+
+		/** Empties the entire hierarchy to the empty state. */
+		void Clear() { AssetNodes.Reset(); CachedNodes.Reset(); }
+
+	private:
+
+		struct FTreeNode
+		{
+			FObjectHierarchyInfo Data;
+			FTreeNode* Parent = nullptr;
+			TArray<TUniquePtr<FTreeNode>> Children;
+
+			explicit FTreeNode(FObjectHierarchyInfo Data)
+				: Data(MoveTemp(Data))
+			{}
+			bool IsAssetNode() const { return Parent == nullptr; }
+		};
+		
+		/**
+		 * Each node represents an asset, e.g. /Game/Maps.Map.
+		 * 
+		 * That means every node in this array is of the form /package/path.assetname (i.e. FSoftObjectPath::AssetPath is set and FSoftObjectPath::SubPathString is empty).
+		 * Remember in general FSoftObjectPath strings look like this: /package/path.assetname[:subpath] where /package/path.assetname is the FSoftObjectPath::AssetPath.
+		 */
+		TArray<TUniquePtr<FTreeNode>> AssetNodes;
+		/** Maps every path to the corresponding node. */
+		TMap<FSoftObjectPath, FTreeNode*> CachedNodes;
+
+		static ETreeTraversalBehavior TraverseTopToBottomInternal(
+			const FObjectHierarchyInfo& Owner,
+			const TArray<TUniquePtr<FTreeNode>>& Nodes,
+			TFunctionRef<ETreeTraversalBehavior(const FChildRelation& Relation)> Callback
+			);
+		static EBreakBehavior TraverseBottomToTopInternal(
+			const FObjectHierarchyInfo& Owner,
+			const TArray<TUniquePtr<FTreeNode>>& Nodes,
+			TFunctionRef<EBreakBehavior(const FChildRelation& Relation)> Callback
+			);
+	};
+}
\ No newline at end of file
diff --git a/Engine/Plugins/Developer/Concert/ConcertSync/ConcertSyncCore/Source/ConcertSyncCore/Public/Misc/ObjectPathOuterIterator.h b/Engine/Plugins/Developer/Concert/ConcertSync/ConcertSyncCore/Source/ConcertSyncCore/Public/Misc/ObjectPathOuterIterator.h
new file mode 100644
--- /dev/null
+++ b/Engine/Plugins/Developer/Concert/ConcertSync/ConcertSyncCore/Source/ConcertSyncCore/Public/Misc/ObjectPathOuterIterator.h
@@ -0,0 +1,62 @@
+﻿// Copyright Epic Games, Inc. All Rights Reserved.
+
+#pragma once
+
+#include "ObjectPathUtils.h"
+
+#include "Misc/Optional.h"
+#include "UObject/SoftObjectPath.h"
+
+namespace UE::ConcertSyncCore
+{
+	/**
+	 * Parses a FSoftObjectPath string and iterates through all outer objects of the Start path.
+	 * 
+	 * Example: Start = /Game/Maps.Map:PersistentLevel.Cube.StaticMeshComponent0.Foo would invoke Callback in this order:
+	 * 1. /Game/Maps.Map:PersistentLevel.Cube.StaticMeshComponent0
+	 * 2. /Game/Maps.Map:PersistentLevel.Cube
+	 * 3. /Game/Maps.Map:PersistentLevel
+	 * 4. /Game/Maps.Map
+	 *
+	 * Example usage:
+	 * const FSoftObjectPath ComponentPath{ TEXT("/Game/Maps.Map:PersistentLevel.Cube.StaticMeshComponent0") };
+	 * for (ConcertSyncCore::FObjectPathOuterIterator It(ComponentPath); It; ++It)
+	 * {
+	 *		const FSoftObjectPath& Path = *It; // 1st iteration = /Game/Maps.Map:PersistentLevel.Cube
+	 *		const FString String = It->ToString(); // 1st iteration = /Game/Maps.Map:PersistentLevel.Cube
+	 * }
+	 */
+	class FObjectPathOuterIterator
+	{
+	public:
+
+		explicit FObjectPathOuterIterator(const FSoftObjectPath& Start)
+			: Current(Start)
+		{
+			Advance();
+		}
+
+		FORCEINLINE void operator++() { Advance(); }
+		
+		FORCEINLINE explicit operator bool() const {  return Current.IsSet();  }
+		FORCEINLINE bool operator !() const  { return !static_cast<bool>(*this); }
+		
+		FORCEINLINE const FSoftObjectPath& operator*() const { check(Current); return *Current; }
+		FORCEINLINE const FSoftObjectPath* operator->() const { return Current.GetPtrOrNull(); }
+		
+		FORCEINLINE bool operator==(const FObjectPathOuterIterator& Right) const { return Current == Right.Current; }
+		FORCEINLINE bool operator!=(const FObjectPathOuterIterator& Right) const { return Current != Right.Current; }
+		
+	private:
+
+		TOptional<FSoftObjectPath> Current;
+
+		void Advance()
+		{
+			if (Current)
+			{
+				Current = GetOuterPath(*Current);
+			}
+		}
+	};
+}
\ No newline at end of file
diff --git a/Engine/Plugins/Developer/Concert/ConcertSync/ConcertSyncCore/Source/ConcertSyncCore/Public/Misc/ObjectPathUtils.h b/Engine/Plugins/Developer/Concert/ConcertSync/ConcertSyncCore/Source/ConcertSyncCore/Public/Misc/ObjectPathUtils.h
new file mode 100644
--- /dev/null
+++ b/Engine/Plugins/Developer/Concert/ConcertSync/ConcertSyncCore/Source/ConcertSyncCore/Public/Misc/ObjectPathUtils.h
@@ -0,0 +1,47 @@
+﻿// Copyright Epic Games, Inc. All Rights Reserved.
+
+#pragma once
+
+#include "Containers/UnrealString.h"
+#include "Misc/Optional.h"
+#include "UObject/SoftObjectPath.h"
+
+namespace UE::ConcertSyncCore
+{
+	/**
+	 * Gets the outer of ObjectPath.
+	 *
+	 * Examples:
+	 * - GetOuterPath(/Game/Map.Map:PersistentLevel.Actor.Component)	= /Game/Map.Map:PersistentLevel.Actor
+	 * - GetOuterPath(/Game/Map.Map:PersistentLevel.Actor)				= /Game/Map.Map:PersistentLevel
+	 * - GetOuterPath(/Game/Map.Map:PersistentLevel)					= /Game/Map.Map
+	 * - GetOuterPath(/Game/Map.Map)									= unset optional
+	 * - GetOuterPath({})												= unset optional
+	 * 
+	 * @return The outer of ObjectPath, if there is one. Unset if ObjectPath points to a package.
+	 */
+	TOptional<FSoftObjectPath> GetOuterPath(const FSoftObjectPath& ObjectPath);
+}
+
+namespace UE::ConcertSyncCore
+{
+	inline TOptional<FSoftObjectPath> GetOuterPath(const FSoftObjectPath& ObjectPath)
+	{
+		if (ObjectPath.IsNull() || !ObjectPath.IsSubobject())
+		{
+			return {};
+		}
+		
+		FString SubPathString = ObjectPath.GetSubPathString();
+		const int32 CurrentIndex = SubPathString.Find(TEXT("."), ESearchCase::CaseSensitive, ESearchDir::FromEnd);
+		if (SubPathString.IsValidIndex(CurrentIndex))
+		{
+			SubPathString.LeftInline(CurrentIndex);
+			return FSoftObjectPath{ ObjectPath.GetAssetPath(), *SubPathString };
+		}
+		else
+		{
+			return FSoftObjectPath{ ObjectPath.GetAssetPath(), {} };
+		}
+	}
+}
diff --git a/Engine/Plugins/Developer/Concert/ConcertSync/ConcertSyncCore/Source/ConcertSyncCore/Public/Replication/Misc/ReplicatedObjectHierarchyCache.h b/Engine/Plugins/Developer/Concert/ConcertSync/ConcertSyncCore/Source/ConcertSyncCore/Public/Replication/Misc/ReplicatedObjectHierarchyCache.h
new file mode 100644
--- /dev/null
+++ b/Engine/Plugins/Developer/Concert/ConcertSync/ConcertSyncCore/Source/ConcertSyncCore/Public/Replication/Misc/ReplicatedObjectHierarchyCache.h
@@ -0,0 +1,75 @@
+﻿// Copyright Epic Games, Inc. All Rights Reserved.
+
+#pragma once
+
+#include "Misc/ObjectPathHierarchy.h"
+
+#include "Containers/Array.h"
+#include "Containers/ContainersFwd.h"
+
+struct FConcertObjectInStreamID;
+struct FConcertReplicatedObjectId;
+struct FConcertReplication_ChangeStream_Request;
+struct FConcertReplicationStream;
+struct FConcertReplication_Join_Request;
+
+namespace UE::ConcertSyncCore
+{
+	/**
+	 * Holds the outer hierarchy of all objects registered in any stream.
+	 * Updates the hierarchy when a relevant event changing stream content happens.
+	 */
+	class CONCERTSYNCCORE_API FReplicatedObjectHierarchyCache : FObjectPathHierarchy
+	{
+	public:
+
+		// AddObject and RemoveObject are not supposed to be invoked - use the dedicated OnX events below.
+		using FObjectPathHierarchy::TraverseTopToBottom;
+		using FObjectPathHierarchy::TraverseBottomToTop;
+		using FObjectPathHierarchy::IsInHierarchy;
+		using FObjectPathHierarchy::HasChildren;
+		using FObjectPathHierarchy::IsEmpty;
+
+		/** Removes all saved state. */
+		void Clear();
+
+		/**
+		 * Checks whether the object is registered by any client (except for this in IgnoredClients).
+		 *
+		 * This function ignores implicit knowledge of the hierarchy.
+		 * For example if you register ONLY /Game/Maps.Map:PersistentLevel.Cube.StaticMeshComponent0, then
+		 * - IsObjectReferencedDirectly(/Game/Maps.Map:PersistentLevel.Cube.StaticMeshComponent0) == true,
+		 * - IsObjectReferencedDirectly(/Game/Maps.Map:PersistentLevel.Cube) == false
+		 * but e.g. TraverseTopToBottom would list both paths.
+		 */
+		bool IsObjectReferencedDirectly(const FSoftObjectPath& ObjectPath, TConstArrayView<FGuid> IgnoredClients = {}) const;
+
+		// Validated network events that modify the known hierarchy of objects.
+		void OnJoin(const FGuid& ClientId, const FConcertReplication_Join_Request& Request);
+		void OnPreClientLeft(const FGuid& ClientId, TConstArrayView<FConcertReplicationStream> Streams);
+		void OnChangeStreams(const FGuid& ClientId, TConstArrayView<FConcertObjectInStreamID> AddedObjects, TConstArrayView<FConcertObjectInStreamID> RemovedObjects);
+
+	private:
+
+		struct FStreamReferencer
+		{
+			FGuid ClientId;
+			FGuid StreamId;
+
+			friend bool operator==(const FStreamReferencer& Left, const FStreamReferencer& Right) { return Left.ClientId == Right.ClientId && Left.StreamId == Right.StreamId; }
+			friend bool operator!=(const FStreamReferencer& Left, const FStreamReferencer& Right) { return !(Left == Right); }
+		};
+		
+		struct FObjectMetaData
+		{
+			/** Keeps track of all clients referencing the stream. */
+			TArray<FStreamReferencer> ReferencingStreams;
+		};
+		
+		TMap<FSoftObjectPath, FObjectMetaData> ObjectMetaData;
+
+		void AddObjectInternal(const FConcertReplicatedObjectId& ObjectInfo);
+		void RemoveObjectInternal(const FConcertReplicatedObjectId& Object);
+		bool RemoveMetaData(const FConcertReplicatedObjectId& Object);
+	};
+}
diff --git a/Engine/Plugins/Developer/Concert/ConcertSync/ConcertSyncTest/Source/ConcertSyncTest/Private/Misc/ObjectPathHierarchy.spec.cpp b/Engine/Plugins/Developer/Concert/ConcertSync/ConcertSyncTest/Source/ConcertSyncTest/Private/Misc/ObjectPathHierarchy.spec.cpp
new file mode 100644
--- /dev/null
+++ b/Engine/Plugins/Developer/Concert/ConcertSync/ConcertSyncTest/Source/ConcertSyncTest/Private/Misc/ObjectPathHierarchy.spec.cpp
@@ -0,0 +1,526 @@
+﻿// Copyright Epic Games, Inc. All Rights Reserved.
+
+#include "Misc/AutomationTest.h"
+
+#include "Misc/AutomationTest.h"
+#include "Misc/Paths.h"
+#include "HAL/FileManager.h"
+#include "Misc/ObjectPathHierarchy.h"
+#include "Misc/ObjectPathUtils.h"
+
+namespace UE::ConcertSyncTests
+{
+	/** Tests functions in ObjectPathUtils.h */
+	BEGIN_DEFINE_SPEC(FObjectPathHierarchySpec, "Editor.Concert.Components.ObjectPathHierarchy", EAutomationTestFlags::EditorContext | EAutomationTestFlags::EngineFilter)
+		ConcertSyncCore::FObjectPathHierarchy Hierarchy;
+	END_DEFINE_SPEC(FObjectPathHierarchySpec);
+
+	void FObjectPathHierarchySpec::Define()
+	{
+		AfterEach([this]() { Hierarchy.Clear(); });
+		
+		Describe("Explicit hierarchy", [this]()
+		{
+			BeforeEach([this]()
+			{
+				Hierarchy.AddObject(FSoftObjectPath{ TEXT("/Game/Maps.Map") });
+				Hierarchy.AddObject(FSoftObjectPath{ TEXT("/Game/Maps.Map:PersistentLevel") });
+				Hierarchy.AddObject(FSoftObjectPath{ TEXT("/Game/Maps.Map:PersistentLevel.Cube") });
+				Hierarchy.AddObject(FSoftObjectPath{ TEXT("/Game/Maps.Map:PersistentLevel.Cube.StaticMeshComponent0") });
+				Hierarchy.AddObject(FSoftObjectPath{ TEXT("/Game/Maps.Map:PersistentLevel.Cube.StaticMeshComponent1") });
+				Hierarchy.AddObject(FSoftObjectPath{ TEXT("/Game/Maps.Map:PersistentLevel.Sphere") });
+				Hierarchy.AddObject(FSoftObjectPath{ TEXT("/Game/Maps.Map:PersistentLevel.Sphere.StaticMeshComponent0") });
+			});
+
+			It("IsInHierarchy", [this]()
+			{
+				const auto TestIsInHierarchy = [this](const TCHAR* Path, ConcertSyncCore::EHierarchyObjectType Expected = ConcertSyncCore::EHierarchyObjectType::Explicit)
+				{
+					const TOptional<ConcertSyncCore::EHierarchyObjectType> State = Hierarchy.IsInHierarchy(FSoftObjectPath{ Path });
+					TestTrue(FString::Printf(TEXT("Is in hierarchy: %s"), Path), State.IsSet() && *State == Expected);
+				};
+				const auto TestIsNotInHierarchy = [this](const TCHAR* Path)
+				{
+					TestFalse(FString::Printf(TEXT("Not in hierarchy: %s"), Path), Hierarchy.IsInHierarchy(FSoftObjectPath{ Path }).IsSet());
+				};
+
+				TestIsInHierarchy(TEXT("/Game/Maps.Map"));
+				TestIsInHierarchy(TEXT("/Game/Maps.Map:PersistentLevel"));
+				TestIsInHierarchy(TEXT("/Game/Maps.Map:PersistentLevel.Cube"));
+				TestIsInHierarchy(TEXT("/Game/Maps.Map:PersistentLevel.Cube.StaticMeshComponent0"));
+				TestIsInHierarchy(TEXT("/Game/Maps.Map:PersistentLevel.Cube.StaticMeshComponent1"));
+				TestIsInHierarchy(TEXT("/Game/Maps.Map:PersistentLevel.Sphere"));
+				TestIsInHierarchy(TEXT("/Game/Maps.Map:PersistentLevel.Sphere.StaticMeshComponent0"));
+
+				Hierarchy.RemoveObject(FSoftObjectPath{ TEXT("/Game/Maps.Map:PersistentLevel.Sphere") });
+				TestIsInHierarchy(TEXT("/Game/Maps.Map"));
+				TestIsInHierarchy(TEXT("/Game/Maps.Map:PersistentLevel"));
+				TestIsInHierarchy(TEXT("/Game/Maps.Map:PersistentLevel.Cube"));
+				TestIsInHierarchy(TEXT("/Game/Maps.Map:PersistentLevel.Cube.StaticMeshComponent0"));
+				TestIsInHierarchy(TEXT("/Game/Maps.Map:PersistentLevel.Cube.StaticMeshComponent1"));
+				// Actual change:
+				TestIsInHierarchy(TEXT("/Game/Maps.Map:PersistentLevel.Sphere"), ConcertSyncCore::EHierarchyObjectType::Implicit);
+				TestIsInHierarchy(TEXT("/Game/Maps.Map:PersistentLevel.Sphere.StaticMeshComponent0"));
+
+				Hierarchy.RemoveObject(FSoftObjectPath{ TEXT("/Game/Maps.Map:PersistentLevel.Sphere.StaticMeshComponent0") });
+				TestIsInHierarchy(TEXT("/Game/Maps.Map"));
+				TestIsInHierarchy(TEXT("/Game/Maps.Map:PersistentLevel"));
+				TestIsInHierarchy(TEXT("/Game/Maps.Map:PersistentLevel.Cube"));
+				TestIsInHierarchy(TEXT("/Game/Maps.Map:PersistentLevel.Cube.StaticMeshComponent0"));
+				TestIsInHierarchy(TEXT("/Game/Maps.Map:PersistentLevel.Cube.StaticMeshComponent1"));
+				// Actual change:
+				TestIsNotInHierarchy(TEXT("/Game/Maps.Map:PersistentLevel.Sphere"));
+				TestIsNotInHierarchy(TEXT("/Game/Maps.Map:PersistentLevel.Sphere.StaticMeshComponent0"));
+			});
+
+			It("HasChildren", [this]()
+			{
+				const auto TestHasChildren = [this](const TCHAR* Path)
+				{
+					TestTrue(FString::Printf(TEXT("Has children: %s"), Path), Hierarchy.HasChildren(FSoftObjectPath{ Path }));
+				};
+				const auto TestHasNoChildren = [this](const TCHAR* Path)
+				{
+					TestFalse(FString::Printf(TEXT("Has no children: %s"), Path), Hierarchy.HasChildren(FSoftObjectPath{ Path }));
+				};
+
+				TestHasChildren(TEXT("/Game/Maps.Map"));
+				TestHasChildren(TEXT("/Game/Maps.Map:PersistentLevel"));
+				TestHasChildren(TEXT("/Game/Maps.Map:PersistentLevel.Cube"));
+				TestHasNoChildren(TEXT("/Game/Maps.Map:PersistentLevel.Cube.StaticMeshComponent0"));
+				TestHasNoChildren(TEXT("/Game/Maps.Map:PersistentLevel.Cube.StaticMeshComponent1"));
+				TestHasChildren(TEXT("/Game/Maps.Map:PersistentLevel.Sphere"));
+				TestHasNoChildren(TEXT("/Game/Maps.Map:PersistentLevel.Sphere.StaticMeshComponent0"));
+
+				Hierarchy.RemoveObject(FSoftObjectPath{ TEXT("/Game/Maps.Map:PersistentLevel.Sphere") });
+				TestHasChildren(TEXT("/Game/Maps.Map"));
+				TestHasChildren(TEXT("/Game/Maps.Map:PersistentLevel"));
+				TestHasChildren(TEXT("/Game/Maps.Map:PersistentLevel.Cube"));
+				TestHasNoChildren(TEXT("/Game/Maps.Map:PersistentLevel.Cube.StaticMeshComponent0"));
+				TestHasNoChildren(TEXT("/Game/Maps.Map:PersistentLevel.Cube.StaticMeshComponent1"));
+				// Actual change:
+				TestHasChildren(TEXT("/Game/Maps.Map:PersistentLevel.Sphere"));
+				TestHasNoChildren(TEXT("/Game/Maps.Map:PersistentLevel.Sphere.StaticMeshComponent0"));
+
+				Hierarchy.RemoveObject(FSoftObjectPath{ TEXT("/Game/Maps.Map:PersistentLevel.Sphere.StaticMeshComponent0") });
+				TestHasChildren(TEXT("/Game/Maps.Map"));
+				TestHasChildren(TEXT("/Game/Maps.Map:PersistentLevel"));
+				TestHasChildren(TEXT("/Game/Maps.Map:PersistentLevel.Cube"));
+				TestHasNoChildren(TEXT("/Game/Maps.Map:PersistentLevel.Cube.StaticMeshComponent0"));
+				TestHasNoChildren(TEXT("/Game/Maps.Map:PersistentLevel.Cube.StaticMeshComponent1"));
+				// Actual change:
+				TestHasNoChildren(TEXT("/Game/Maps.Map:PersistentLevel.Sphere"));
+				TestHasNoChildren(TEXT("/Game/Maps.Map:PersistentLevel.Sphere.StaticMeshComponent0"));
+			});
+			
+			It("TraverseTopToBottom (all)", [this]()
+			{
+				TArray<FSoftObjectPath> Visited;
+				Hierarchy.TraverseTopToBottom([this, &Visited](const ConcertSyncCore::FChildRelation& Relation)
+				{
+					Visited.AddUnique(Relation.Parent.Object);
+					Visited.AddUnique(Relation.Child.Object);
+					
+					const TOptional<FSoftObjectPath> ExpectedParent = ConcertSyncCore::GetOuterPath(Relation.Child.Object);
+					TestTrue(TEXT("Is direct child"), ExpectedParent && *ExpectedParent == Relation.Parent.Object);
+					
+					return ConcertSyncCore::ETreeTraversalBehavior::Continue;
+				});
+
+				// Any of the following orders is correct
+				TArray<FSoftObjectPath> PossibleOrder1
+				{
+					FSoftObjectPath(TEXT("/Game/Maps.Map")),
+					FSoftObjectPath(TEXT("/Game/Maps.Map:PersistentLevel")),
+					FSoftObjectPath(TEXT("/Game/Maps.Map:PersistentLevel.Cube")),
+					FSoftObjectPath(TEXT("/Game/Maps.Map:PersistentLevel.Cube.StaticMeshComponent0")),
+					FSoftObjectPath(TEXT("/Game/Maps.Map:PersistentLevel.Cube.StaticMeshComponent1")),
+					FSoftObjectPath(TEXT("/Game/Maps.Map:PersistentLevel.Sphere")),
+					FSoftObjectPath(TEXT("/Game/Maps.Map:PersistentLevel.Sphere.StaticMeshComponent0")),
+				};
+				TArray<FSoftObjectPath> PossibleOrder2
+				{
+					FSoftObjectPath(TEXT("/Game/Maps.Map")),
+					FSoftObjectPath(TEXT("/Game/Maps.Map:PersistentLevel")),
+					FSoftObjectPath(TEXT("/Game/Maps.Map:PersistentLevel.Cube")),
+					FSoftObjectPath(TEXT("/Game/Maps.Map:PersistentLevel.Cube.StaticMeshComponent1")),
+					FSoftObjectPath(TEXT("/Game/Maps.Map:PersistentLevel.Cube.StaticMeshComponent0")),
+					FSoftObjectPath(TEXT("/Game/Maps.Map:PersistentLevel.Sphere")),
+					FSoftObjectPath(TEXT("/Game/Maps.Map:PersistentLevel.Sphere.StaticMeshComponent0")),
+				};
+				TArray<FSoftObjectPath> PossibleOrder3
+				{
+					FSoftObjectPath(TEXT("/Game/Maps.Map")),
+					FSoftObjectPath(TEXT("/Game/Maps.Map:PersistentLevel")),
+					FSoftObjectPath(TEXT("/Game/Maps.Map:PersistentLevel.Sphere")),
+					FSoftObjectPath(TEXT("/Game/Maps.Map:PersistentLevel.Sphere.StaticMeshComponent0")),
+					FSoftObjectPath(TEXT("/Game/Maps.Map:PersistentLevel.Cube")),
+					FSoftObjectPath(TEXT("/Game/Maps.Map:PersistentLevel.Cube.StaticMeshComponent0")),
+					FSoftObjectPath(TEXT("/Game/Maps.Map:PersistentLevel.Cube.StaticMeshComponent1")),
+				};
+				TArray<FSoftObjectPath> PossibleOrder4
+				{
+					FSoftObjectPath(TEXT("/Game/Maps.Map")),
+					FSoftObjectPath(TEXT("/Game/Maps.Map:PersistentLevel")),
+					FSoftObjectPath(TEXT("/Game/Maps.Map:PersistentLevel.Sphere")),
+					FSoftObjectPath(TEXT("/Game/Maps.Map:PersistentLevel.Sphere.StaticMeshComponent0")),
+					FSoftObjectPath(TEXT("/Game/Maps.Map:PersistentLevel.Cube")),
+					FSoftObjectPath(TEXT("/Game/Maps.Map:PersistentLevel.Cube.StaticMeshComponent1")),
+					FSoftObjectPath(TEXT("/Game/Maps.Map:PersistentLevel.Cube.StaticMeshComponent0")),
+				};
+
+				const bool bIsOneOfExpectedOrders = Visited == PossibleOrder1 || Visited == PossibleOrder2 || Visited == PossibleOrder3 || Visited == PossibleOrder4;
+				TestTrue(TEXT("Is correct order"), bIsOneOfExpectedOrders);
+			});
+
+			It("TraverseTopToBottom (SkipSubtree)", [this]()
+			{
+				TArray<FSoftObjectPath> Visited;
+				Hierarchy.TraverseTopToBottom([this, &Visited](const ConcertSyncCore::FChildRelation& Relation)
+				{
+					Visited.AddUnique(Relation.Parent.Object);
+					Visited.AddUnique(Relation.Child.Object);
+					
+					const TOptional<FSoftObjectPath> ExpectedParent = ConcertSyncCore::GetOuterPath(Relation.Child.Object);
+					TestTrue(TEXT("Is direct child"), ExpectedParent && *ExpectedParent == Relation.Parent.Object);
+					
+					return Relation.Child.Object == FSoftObjectPath(TEXT("/Game/Maps.Map:PersistentLevel.Cube"))
+						? ConcertSyncCore::ETreeTraversalBehavior::SkipSubtree
+						: ConcertSyncCore::ETreeTraversalBehavior::Continue;
+				});
+
+				// Any of the following orders is correct
+				TArray<FSoftObjectPath> PossibleOrder1
+				{
+					FSoftObjectPath(TEXT("/Game/Maps.Map")),
+					FSoftObjectPath(TEXT("/Game/Maps.Map:PersistentLevel")),
+					FSoftObjectPath(TEXT("/Game/Maps.Map:PersistentLevel.Cube")),
+					FSoftObjectPath(TEXT("/Game/Maps.Map:PersistentLevel.Sphere")),
+					FSoftObjectPath(TEXT("/Game/Maps.Map:PersistentLevel.Sphere.StaticMeshComponent0"))
+				};
+				TArray<FSoftObjectPath> PossibleOrder2
+				{
+					FSoftObjectPath(TEXT("/Game/Maps.Map")),
+					FSoftObjectPath(TEXT("/Game/Maps.Map:PersistentLevel")),
+					FSoftObjectPath(TEXT("/Game/Maps.Map:PersistentLevel.Sphere")),
+					FSoftObjectPath(TEXT("/Game/Maps.Map:PersistentLevel.Sphere.StaticMeshComponent0")),
+					FSoftObjectPath(TEXT("/Game/Maps.Map:PersistentLevel.Cube"))
+				};
+
+				const bool bIsOneOfExpectedOrders = Visited == PossibleOrder1 || Visited == PossibleOrder2;
+				TestTrue(TEXT("Is correct order"), bIsOneOfExpectedOrders);
+			});
+			
+			It("TraverseTopToBottom (/Game/Maps.Map:PersistentLevel.Sphere)", [this]()
+			{
+				bool bVisited = false;
+				Hierarchy.TraverseTopToBottom([this, &bVisited](const ConcertSyncCore::FChildRelation& Relation)
+				{
+					if (bVisited)
+					{
+						AddError(TEXT("Expected 1 invocation"));
+						return ConcertSyncCore::ETreeTraversalBehavior::Break;
+					}
+
+					bVisited = true;
+					TestEqual(TEXT("Parent"), Relation.Parent.Object.ToString(), TEXT("/Game/Maps.Map:PersistentLevel.Sphere"));
+					TestEqual(TEXT("Child"), Relation.Child.Object.ToString(), TEXT("/Game/Maps.Map:PersistentLevel.Sphere.StaticMeshComponent0"));
+					TestEqual(TEXT("Parent type"), Relation.Parent.Type, ConcertSyncCore::EHierarchyObjectType::Explicit);
+					TestEqual(TEXT("Child type"), Relation.Child.Type, ConcertSyncCore::EHierarchyObjectType::Explicit);
+					
+					return ConcertSyncCore::ETreeTraversalBehavior::Continue;
+				}, FSoftObjectPath{ TEXT("/Game/Maps.Map:PersistentLevel.Sphere") });
+				TestTrue(TEXT("Visited"), bVisited);
+			});
+
+			It("TraverseBottomToTop (all)", [this]()
+			{
+				TArray<FSoftObjectPath> Visited;
+				int32 NumRootEncountered = 0;
+				Hierarchy.TraverseBottomToTop([this, &Visited, &NumRootEncountered](const ConcertSyncCore::FChildRelation& Relation)
+				{
+					Visited.AddUnique(Relation.Child.Object);
+					
+					const TOptional<FSoftObjectPath> ExpectedParent = ConcertSyncCore::GetOuterPath(Relation.Child.Object);
+					TestTrue(TEXT("Is direct child"), ExpectedParent && *ExpectedParent == Relation.Parent.Object);
+
+					NumRootEncountered += Relation.Parent.Object == FSoftObjectPath(TEXT("/Game/Maps.Map"));
+					return EBreakBehavior::Continue;
+				});
+
+				// Any of the following orders is correct
+				TArray<FSoftObjectPath> PossibleOrder1
+				{
+					FSoftObjectPath(TEXT("/Game/Maps.Map:PersistentLevel.Sphere.StaticMeshComponent0")),
+					FSoftObjectPath(TEXT("/Game/Maps.Map:PersistentLevel.Sphere")),
+					FSoftObjectPath(TEXT("/Game/Maps.Map:PersistentLevel.Cube.StaticMeshComponent1")),
+					FSoftObjectPath(TEXT("/Game/Maps.Map:PersistentLevel.Cube.StaticMeshComponent0")),
+					FSoftObjectPath(TEXT("/Game/Maps.Map:PersistentLevel.Cube")),
+					FSoftObjectPath(TEXT("/Game/Maps.Map:PersistentLevel")),
+					FSoftObjectPath(TEXT("/Game/Maps.Map")),
+				};
+				TArray<FSoftObjectPath> PossibleOrder2
+				{
+					FSoftObjectPath(TEXT("/Game/Maps.Map:PersistentLevel.Sphere.StaticMeshComponent0")),
+					FSoftObjectPath(TEXT("/Game/Maps.Map:PersistentLevel.Sphere")),
+					FSoftObjectPath(TEXT("/Game/Maps.Map:PersistentLevel.Cube.StaticMeshComponent0")),
+					FSoftObjectPath(TEXT("/Game/Maps.Map:PersistentLevel.Cube.StaticMeshComponent1")),
+					FSoftObjectPath(TEXT("/Game/Maps.Map:PersistentLevel.Cube")),
+					FSoftObjectPath(TEXT("/Game/Maps.Map:PersistentLevel")),
+					FSoftObjectPath(TEXT("/Game/Maps.Map")),
+				};
+				TArray<FSoftObjectPath> PossibleOrder3
+				{
+					FSoftObjectPath(TEXT("/Game/Maps.Map:PersistentLevel.Cube.StaticMeshComponent1")),
+					FSoftObjectPath(TEXT("/Game/Maps.Map:PersistentLevel.Cube.StaticMeshComponent0")),
+					FSoftObjectPath(TEXT("/Game/Maps.Map:PersistentLevel.Cube")),
+					FSoftObjectPath(TEXT("/Game/Maps.Map:PersistentLevel.Sphere.StaticMeshComponent0")),
+					FSoftObjectPath(TEXT("/Game/Maps.Map:PersistentLevel.Sphere")),
+					FSoftObjectPath(TEXT("/Game/Maps.Map:PersistentLevel")),
+					FSoftObjectPath(TEXT("/Game/Maps.Map")),
+				};
+				TArray<FSoftObjectPath> PossibleOrder4
+				{
+					FSoftObjectPath(TEXT("/Game/Maps.Map:PersistentLevel.Cube.StaticMeshComponent0")),
+					FSoftObjectPath(TEXT("/Game/Maps.Map:PersistentLevel.Cube.StaticMeshComponent1")),
+					FSoftObjectPath(TEXT("/Game/Maps.Map:PersistentLevel.Cube")),
+					FSoftObjectPath(TEXT("/Game/Maps.Map:PersistentLevel.Sphere.StaticMeshComponent0")),
+					FSoftObjectPath(TEXT("/Game/Maps.Map:PersistentLevel.Sphere")),
+					FSoftObjectPath(TEXT("/Game/Maps.Map:PersistentLevel")),
+					FSoftObjectPath(TEXT("/Game/Maps.Map")),
+				};
+
+				TestEqual(TEXT("Visited /Game/Maps.Map once"), NumRootEncountered, 1);
+				Visited.Add(FSoftObjectPath(TEXT("/Game/Maps.Map")));
+				
+				const bool bIsOneOfExpectedOrders = Visited == PossibleOrder1 || Visited == PossibleOrder2 || Visited == PossibleOrder3 || Visited == PossibleOrder4;
+				TestTrue(TEXT("Is correct order"), bIsOneOfExpectedOrders);
+			});
+
+			It("TraverseBottomToTop (/Game/Maps.Map:PersistentLevel.Sphere)", [this]()
+			{
+				bool bVisited = false;
+				Hierarchy.TraverseBottomToTop([this, &bVisited](const ConcertSyncCore::FChildRelation& Relation)
+				{
+					if (bVisited)
+					{
+						AddError(TEXT("Expected 1 invocation"));
+						return EBreakBehavior::Break;
+					}
+
+					bVisited = true;
+					TestEqual(TEXT("Parent"), Relation.Parent.Object.ToString(), TEXT("/Game/Maps.Map:PersistentLevel.Sphere"));
+					TestEqual(TEXT("Child"), Relation.Child.Object.ToString(), TEXT("/Game/Maps.Map:PersistentLevel.Sphere.StaticMeshComponent0"));
+					TestEqual(TEXT("Parent type"), Relation.Parent.Type, ConcertSyncCore::EHierarchyObjectType::Explicit);
+					TestEqual(TEXT("Child type"), Relation.Child.Type, ConcertSyncCore::EHierarchyObjectType::Explicit);
+					
+					return EBreakBehavior::Continue;
+				}, FSoftObjectPath{ TEXT("/Game/Maps.Map:PersistentLevel.Sphere") });
+				TestTrue(TEXT("Visited"), bVisited);
+			});
+		});
+
+		Describe("Implicit hierarchy", [this]()
+		{
+			BeforeEach([this]()
+			{
+				Hierarchy.AddObject(FSoftObjectPath{ TEXT("/Game/Maps.Map:PersistentLevel.Sphere") });
+				Hierarchy.AddObject(FSoftObjectPath{ TEXT("/Game/Maps.Map:PersistentLevel.Sphere.StaticMeshComponent0.Subobject") });
+			});
+
+			It("IsInHierarchy", [this]()
+			{
+				using namespace ConcertSyncCore;
+				const auto TestIsInHierarchy = [this](const TCHAR* Path, EHierarchyObjectType Expected )
+				{
+					const TOptional<EHierarchyObjectType> State = Hierarchy.IsInHierarchy(FSoftObjectPath{ Path });
+					TestTrue(FString::Printf(TEXT("Is in hierarchy: %s"), Path), State.IsSet() && *State == Expected);
+				};
+
+				TestIsInHierarchy(TEXT("/Game/Maps.Map"), EHierarchyObjectType::Implicit);
+				TestIsInHierarchy(TEXT("/Game/Maps.Map:PersistentLevel"), EHierarchyObjectType::Implicit);
+				TestIsInHierarchy(TEXT("/Game/Maps.Map:PersistentLevel.Sphere"), EHierarchyObjectType::Explicit);
+				TestIsInHierarchy(TEXT("/Game/Maps.Map:PersistentLevel.Sphere.StaticMeshComponent0"), EHierarchyObjectType::Implicit);
+				TestIsInHierarchy(TEXT("/Game/Maps.Map:PersistentLevel.Sphere.StaticMeshComponent0.Subobject"), EHierarchyObjectType::Explicit);
+			});
+			
+			It("HasChildren", [this]()
+			{
+				using namespace ConcertSyncCore;
+				const auto TestHasChildren = [this](const TCHAR* Path)
+				{
+					TestTrue(FString::Printf(TEXT("Has children: %s"), Path), Hierarchy.HasChildren(FSoftObjectPath{ Path }));
+				};
+				const auto TestHasNoChildren = [this](const TCHAR* Path)
+				{
+					TestFalse(FString::Printf(TEXT("Has no children: %s"), Path), Hierarchy.HasChildren(FSoftObjectPath{ Path }));
+				};
+
+				TestHasChildren(TEXT("/Game/Maps.Map"));
+				TestHasChildren(TEXT("/Game/Maps.Map:PersistentLevel"));
+				TestHasChildren(TEXT("/Game/Maps.Map:PersistentLevel.Sphere"));
+				TestHasChildren(TEXT("/Game/Maps.Map:PersistentLevel.Sphere.StaticMeshComponent0"));
+				TestHasNoChildren(TEXT("/Game/Maps.Map:PersistentLevel.Sphere.StaticMeshComponent0.Subobject"));
+			});
+
+			It("TraverseTopToBottom (all)", [this]()
+			{
+				TArray<FSoftObjectPath> Visited;
+				Hierarchy.TraverseTopToBottom([this, &Visited](const ConcertSyncCore::FChildRelation& Relation)
+				{
+					Visited.AddUnique(Relation.Parent.Object);
+					Visited.AddUnique(Relation.Child.Object);
+					
+					const TOptional<FSoftObjectPath> ExpectedParent = ConcertSyncCore::GetOuterPath(Relation.Child.Object);
+					TestTrue(TEXT("Is direct child"), ExpectedParent && *ExpectedParent == Relation.Parent.Object);
+					
+					return ConcertSyncCore::ETreeTraversalBehavior::Continue;
+				});
+				
+				TArray<FSoftObjectPath> ExpectedOrder
+				{
+					FSoftObjectPath(TEXT("/Game/Maps.Map")),
+					FSoftObjectPath(TEXT("/Game/Maps.Map:PersistentLevel")),
+					FSoftObjectPath(TEXT("/Game/Maps.Map:PersistentLevel.Sphere")),
+					FSoftObjectPath(TEXT("/Game/Maps.Map:PersistentLevel.Sphere.StaticMeshComponent0")),
+					FSoftObjectPath(TEXT("/Game/Maps.Map:PersistentLevel.Sphere.StaticMeshComponent0.Subobject"))
+				};
+				TestEqual(TEXT("Visted in expected order"), Visited, ExpectedOrder);
+			});
+			
+			It("TraverseBottomToTop (all)", [this]()
+			{
+				TArray<FSoftObjectPath> Visited;
+				Hierarchy.TraverseBottomToTop([this, &Visited](const ConcertSyncCore::FChildRelation& Relation)
+				{
+					Visited.AddUnique(Relation.Child.Object);
+					Visited.AddUnique(Relation.Parent.Object);
+					
+					const TOptional<FSoftObjectPath> ExpectedParent = ConcertSyncCore::GetOuterPath(Relation.Child.Object);
+					TestTrue(TEXT("Is direct child"), ExpectedParent && *ExpectedParent == Relation.Parent.Object);
+					
+					return EBreakBehavior::Continue;
+				});
+
+				TArray<FSoftObjectPath> ExpectedOrder
+				{
+					FSoftObjectPath(TEXT("/Game/Maps.Map:PersistentLevel.Sphere.StaticMeshComponent0.Subobject")),
+					FSoftObjectPath(TEXT("/Game/Maps.Map:PersistentLevel.Sphere.StaticMeshComponent0")),
+					FSoftObjectPath(TEXT("/Game/Maps.Map:PersistentLevel.Sphere")),
+					FSoftObjectPath(TEXT("/Game/Maps.Map:PersistentLevel")),
+					FSoftObjectPath(TEXT("/Game/Maps.Map"))
+				};
+				TestEqual(TEXT("Visted in expected order"), Visited, ExpectedOrder);
+			});
+		});
+
+		Describe("Add & remove all (world)", [this]()
+		{
+			AfterEach([this]()
+			{
+				TestTrue("IsEmpty()", Hierarchy.IsEmpty());
+				TestFalse("HasChildren(PersistentLevel)", Hierarchy.HasChildren(FSoftObjectPath{ TEXT("/Game/Maps.Map:PersistentLevel") }));
+				TestFalse("HasChildren(Cube)", Hierarchy.HasChildren(FSoftObjectPath{ TEXT("/Game/Maps.Map:PersistentLevel.Cube") }));
+				TestFalse("IsInHierarchy(PersistentLevel)", Hierarchy.IsInHierarchy(FSoftObjectPath{ TEXT("/Game/Maps.Map:PersistentLevel") }).IsSet());
+				TestFalse("IsInHierarchy(Cube)", Hierarchy.IsInHierarchy(FSoftObjectPath{ TEXT("/Game/Maps.Map:PersistentLevel.Cube") }).IsSet());
+			});
+
+			It("Add cube, remove cube.", [this]()
+			{
+				Hierarchy.AddObject(FSoftObjectPath{ TEXT("/Game/Maps.Map:PersistentLevel.Cube") });
+				Hierarchy.RemoveObject(FSoftObjectPath{ TEXT("/Game/Maps.Map:PersistentLevel.Cube") });
+			});
+
+			// Test various orders in which may be added & removed.
+
+			It("Add level, add cube, remove level, remove cube.", [this]()
+			{
+				Hierarchy.AddObject(FSoftObjectPath{ TEXT("/Game/Maps.Map:PersistentLevel") });
+				Hierarchy.AddObject(FSoftObjectPath{ TEXT("/Game/Maps.Map:PersistentLevel.Cube") });
+
+				Hierarchy.RemoveObject(FSoftObjectPath{ TEXT("/Game/Maps.Map:PersistentLevel") });
+				Hierarchy.RemoveObject(FSoftObjectPath{ TEXT("/Game/Maps.Map:PersistentLevel.Cube") });
+			});
+			It("Add level, add cube, remove cube, remove level.", [this]()
+			{
+				Hierarchy.AddObject(FSoftObjectPath{ TEXT("/Game/Maps.Map:PersistentLevel") });
+				Hierarchy.AddObject(FSoftObjectPath{ TEXT("/Game/Maps.Map:PersistentLevel.Cube") });
+
+				Hierarchy.RemoveObject(FSoftObjectPath{ TEXT("/Game/Maps.Map:PersistentLevel.Cube") });
+				Hierarchy.RemoveObject(FSoftObjectPath{ TEXT("/Game/Maps.Map:PersistentLevel") });
+			});
+
+			It("Add cube, add level, remove level, remove cube.", [this]()
+			{
+				Hierarchy.AddObject(FSoftObjectPath{ TEXT("/Game/Maps.Map:PersistentLevel.Cube") });
+				Hierarchy.AddObject(FSoftObjectPath{ TEXT("/Game/Maps.Map:PersistentLevel") });
+
+				Hierarchy.RemoveObject(FSoftObjectPath{ TEXT("/Game/Maps.Map:PersistentLevel") });
+				Hierarchy.RemoveObject(FSoftObjectPath{ TEXT("/Game/Maps.Map:PersistentLevel.Cube") });
+			});
+			It("Add cube, add level, remove cube, remove level.", [this]()
+			{
+				Hierarchy.AddObject(FSoftObjectPath{ TEXT("/Game/Maps.Map:PersistentLevel.Cube") });
+				Hierarchy.AddObject(FSoftObjectPath{ TEXT("/Game/Maps.Map:PersistentLevel") });
+
+				Hierarchy.RemoveObject(FSoftObjectPath{ TEXT("/Game/Maps.Map:PersistentLevel.Cube") });
+				Hierarchy.RemoveObject(FSoftObjectPath{ TEXT("/Game/Maps.Map:PersistentLevel") });
+			});
+		});
+
+		// The difference to the world tests above is that we're directly removing an object that is placed in FObjectPathHierarchy::AssetNodes
+		Describe("Add & remove all (asset)", [this]()
+		{
+			AfterEach([this]()
+			{
+				TestTrue("IsEmpty()", Hierarchy.IsEmpty());
+				TestFalse("HasChildren(Root)", Hierarchy.HasChildren(FSoftObjectPath{ TEXT("/Engine/Transient.Root") }));
+				TestFalse("HasChildren(Subobject)", Hierarchy.HasChildren(FSoftObjectPath{ TEXT("/Engine/Transient.Root:Subobject") }));
+				TestFalse("IsInHierarchy(Root)", Hierarchy.IsInHierarchy(FSoftObjectPath{ TEXT("/Engine/Transient.Root") }).IsSet());
+				TestFalse("IsInHierarchy(Subobject)", Hierarchy.IsInHierarchy(FSoftObjectPath{ TEXT("/Engine/Transient.Root:Subobject") }).IsSet());
+			});
+
+			It("Add subobject, remove subobject.", [this]()
+			{
+				Hierarchy.AddObject(FSoftObjectPath{ TEXT("/Engine/Transient.Root:Subobject") });
+				Hierarchy.RemoveObject(FSoftObjectPath{ TEXT("/Engine/Transient.Root:Subobject") });
+			});
+
+			// Test various orders in which may be added & removed.
+
+			It("Add root, add subobject, remove root, remove subobject.", [this]()
+			{
+				Hierarchy.AddObject(FSoftObjectPath{ TEXT("/Engine/Transient.Root") });
+				Hierarchy.AddObject(FSoftObjectPath{ TEXT("/Engine/Transient.Root:Subobject") });
+
+				Hierarchy.RemoveObject(FSoftObjectPath{ TEXT("/Engine/Transient.Root") });
+				Hierarchy.RemoveObject(FSoftObjectPath{ TEXT("/Engine/Transient.Root:Subobject") });
+			});
+			It("Add root, add subobject, remove subobject, remove root.", [this]()
+			{
+				Hierarchy.AddObject(FSoftObjectPath{ TEXT("/Engine/Transient.Root") });
+				Hierarchy.AddObject(FSoftObjectPath{ TEXT("/Engine/Transient.Root:Subobject") });
+
+				Hierarchy.RemoveObject(FSoftObjectPath{ TEXT("/Engine/Transient.Root:Subobject") });
+				Hierarchy.RemoveObject(FSoftObjectPath{ TEXT("/Engine/Transient.Root") });
+			});
+
+			It("Add subobject, add root, remove root, remove subobject.", [this]()
+			{
+				Hierarchy.AddObject(FSoftObjectPath{ TEXT("/Engine/Transient.Root:Subobject") });
+				Hierarchy.AddObject(FSoftObjectPath{ TEXT("/Engine/Transient.Root") });
+
+				Hierarchy.RemoveObject(FSoftObjectPath{ TEXT("/Engine/Transient.Root") });
+				Hierarchy.RemoveObject(FSoftObjectPath{ TEXT("/Engine/Transient.Root:Subobject") });
+			});
+			It("Add subobject, add root, remove subobject, remove root.", [this]()
+			{
+				Hierarchy.AddObject(FSoftObjectPath{ TEXT("/Engine/Transient.Root:Subobject") });
+				Hierarchy.AddObject(FSoftObjectPath{ TEXT("/Engine/Transient.Root") });
+
+				Hierarchy.RemoveObject(FSoftObjectPath{ TEXT("/Engine/Transient.Root:Subobject") });
+				Hierarchy.RemoveObject(FSoftObjectPath{ TEXT("/Engine/Transient.Root") });
+			});
+		});
+	}
+}
diff --git a/Engine/Plugins/Developer/Concert/ConcertSync/ConcertSyncTest/Source/ConcertSyncTest/Private/Misc/ObjectPathUtils.spec.cpp b/Engine/Plugins/Developer/Concert/ConcertSync/ConcertSyncTest/Source/ConcertSyncTest/Private/Misc/ObjectPathUtils.spec.cpp
new file mode 100644
--- /dev/null
+++ b/Engine/Plugins/Developer/Concert/ConcertSync/ConcertSyncTest/Source/ConcertSyncTest/Private/Misc/ObjectPathUtils.spec.cpp
@@ -0,0 +1,77 @@
+﻿// Copyright Epic Games, Inc. All Rights Reserved.
+
+#include "Misc/AutomationTest.h"
+
+#include "Misc/AutomationTest.h"
+#include "Misc/Paths.h"
+#include "HAL/FileManager.h"
+#include "Misc/ObjectPathOuterIterator.h"
+#include "Misc/ObjectPathUtils.h"
+
+namespace UE::ConcertSyncTests
+{
+	/** Tests functions in ObjectPathUtils.h */
+	BEGIN_DEFINE_SPEC(FObjectPathUtilsSpec, "Editor.Concert.Components.ObjectPathUtils", EAutomationTestFlags::EditorContext | EAutomationTestFlags::EngineFilter)
+	END_DEFINE_SPEC(FObjectPathUtilsSpec);
+
+	void FObjectPathUtilsSpec::Define()
+	{
+		It("FObjectPathOuterIterator", [this]()
+		{
+			int32 NumberOfInvocations = 0;
+			const FSoftObjectPath ComponentPath{ TEXT("/Game/Map.Map:PersistentLevel.Actor.Component") };
+			for (ConcertSyncCore::FObjectPathOuterIterator It(ComponentPath); It; ++It)
+			{
+				++NumberOfInvocations;
+				
+				const FString PathString = It->ToString();
+				switch (NumberOfInvocations)
+				{
+				case 1: TestEqual(TEXT("Equal to /Game/Map.Map:PersistentLevel.Actor"), PathString, TEXT("/Game/Map.Map:PersistentLevel.Actor")); break;
+				case 2: TestEqual(TEXT("Equal to /Game/Map.Map:PersistentLevel"), PathString, TEXT("/Game/Map.Map:PersistentLevel")); break;
+				case 3: TestEqual(TEXT("Equal to /Game/Map.Map"), PathString, TEXT("/Game/Map.Map")); break;
+				default: AddError(TEXT("Too many invocations"));
+				}
+			}
+			TestEqual(TEXT("Invoked exactly 3 times"), NumberOfInvocations, 3);
+
+			for (ConcertSyncCore::FObjectPathOuterIterator It(FSoftObjectPath{ TEXT("/Game/Map.Map") }); It; ++It)
+			{
+				AddError(TEXT("Assets do not have any outers"));
+			}
+			for (ConcertSyncCore::FObjectPathOuterIterator It(FSoftObjectPath{}); It; ++It)
+			{
+				AddError(TEXT("Null iteration"));
+			}
+		});
+
+		It("GetOuterPath", [this]()
+		{
+			const TOptional<FSoftObjectPath> ActorPath = ConcertSyncCore::GetOuterPath(FSoftObjectPath{ TEXT("/Game/Map.Map:PersistentLevel.Actor.Component") });
+			TestEqual(
+				TEXT("/Game/Map.Map:PersistentLevel.Actor.Component"),
+				ActorPath ? ActorPath->ToString() : TEXT(""),
+				TEXT("/Game/Map.Map:PersistentLevel.Actor")
+				);
+			
+			const TOptional<FSoftObjectPath> PersistentLevelPath = ConcertSyncCore::GetOuterPath(FSoftObjectPath{ TEXT("/Game/Map.Map:PersistentLevel.Actor") });
+			TestEqual(
+				TEXT("/Game/Map.Map:PersistentLevel.Actor"),
+				PersistentLevelPath ? PersistentLevelPath->ToString() : TEXT(""),
+				TEXT("/Game/Map.Map:PersistentLevel")
+				);
+			
+			const TOptional<FSoftObjectPath> PackagePath = ConcertSyncCore::GetOuterPath(FSoftObjectPath{ TEXT("/Game/Map.Map:PersistentLevel") });
+			TestEqual(
+				TEXT("/Game/Map.Map:PersistentLevel"),
+				PackagePath ? PackagePath->ToString() : TEXT(""),
+				TEXT("/Game/Map.Map")
+				);
+
+			const TOptional<FSoftObjectPath> NoPath = ConcertSyncCore::GetOuterPath(FSoftObjectPath{ TEXT("/Game/Map.Map") });
+			TestFalse(TEXT("/Game/Map.Map"), NoPath.IsSet());
+			const TOptional<FSoftObjectPath> NullPath = ConcertSyncCore::GetOuterPath(nullptr);
+			TestFalse(TEXT("Null"), NullPath.IsSet());
+		});
+	}
+}
diff --git a/Engine/Plugins/Developer/Concert/ConcertSync/ConcertSyncTest/Source/ConcertSyncTest/Private/Replication/Components/ReplicatedObjectHierarchyCache.spec.cpp b/Engine/Plugins/Developer/Concert/ConcertSync/ConcertSyncTest/Source/ConcertSyncTest/Private/Replication/Components/ReplicatedObjectHierarchyCache.spec.cpp
new file mode 100644
--- /dev/null
+++ b/Engine/Plugins/Developer/Concert/ConcertSync/ConcertSyncTest/Source/ConcertSyncTest/Private/Replication/Components/ReplicatedObjectHierarchyCache.spec.cpp
@@ -0,0 +1,140 @@
+// Copyright Epic Games, Inc. All Rights Reserved.
+
+#include "Misc/AutomationTest.h"
+
+#include "Misc/AutomationTest.h"
+#include "Misc/Paths.h"
+#include "HAL/FileManager.h"
+#include "Replication/Misc/ReplicatedObjectHierarchyCache.h"
+#include "Replication/Util/Spec/ObjectTestReplicator.h"
+#include "Replication/Util/Spec/ReplicationClient.h"
+
+namespace UE::ConcertSyncTests::Replication
+{
+	/** Tests functions in ObjectPathUtils.h */
+	BEGIN_DEFINE_SPEC(FReplicatedObjectHierarchyCacheSpec, "Editor.Concert.Replication.Components.ReplicatedObjectHierarchyCache", EAutomationTestFlags::EditorContext | EAutomationTestFlags::EngineFilter)
+		FGuid ClientId_One = FGuid::NewGuid(); 
+		FGuid ClientId_Two = FGuid::NewGuid(); 
+		FGuid ClientId_Three = FGuid::NewGuid(); 
+	
+		ConcertSyncCore::FReplicatedObjectHierarchyCache Cache;
+		TSharedPtr<FObjectTestReplicator> ObjectReplicator;
+		TSharedPtr<FObjectTestReplicator> SubobjectReplicator;
+
+		void TestObjectInCache()
+		{
+			TestFalse(TEXT("IsObjectReferencedDirectly(Package)"), Cache.IsObjectReferencedDirectly(ObjectReplicator->TestObject->GetOuter()));
+			TestFalse(TEXT("IsInHierarchy(Package)"), Cache.IsInHierarchy(ObjectReplicator->TestObject->GetOuter()).IsSet());
+			
+			TestTrue(TEXT("IsObjectReferencedDirectly(Object)"), Cache.IsObjectReferencedDirectly(ObjectReplicator->TestObject));
+			TestTrue(TEXT("IsInHierarchy(Object)"), Cache.IsInHierarchy(ObjectReplicator->TestObject).IsSet());
+		}
+	
+		void TestSubobjectInCache()
+		{
+			TestFalse(TEXT("IsObjectReferencedDirectly(Package)"), Cache.IsObjectReferencedDirectly(ObjectReplicator->TestObject->GetOuter()));
+			TestFalse(TEXT("IsInHierarchy(Package)"), Cache.IsInHierarchy(ObjectReplicator->TestObject->GetOuter()).IsSet());
+			
+			TestTrue(TEXT("IsObjectReferencedDirectly(Subobject)"), Cache.IsObjectReferencedDirectly(SubobjectReplicator->TestObject));
+			TestTrue(TEXT("IsInHierarchy(Subobject)"), Cache.IsInHierarchy(SubobjectReplicator->TestObject).IsSet());
+		}
+	END_DEFINE_SPEC(FReplicatedObjectHierarchyCacheSpec);
+
+	void FReplicatedObjectHierarchyCacheSpec::Define()
+	{
+		BeforeEach([this]()
+		{
+			ObjectReplicator = MakeShared<FObjectTestReplicator>();
+			SubobjectReplicator = ObjectReplicator->CreateSubobjectReplicator();
+		});
+		AfterEach([this]()
+		{
+			Cache.Clear();
+			ObjectReplicator.Reset();
+			SubobjectReplicator.Reset();
+		});
+		
+		Describe("Single client", [this]()
+		{
+			It("OnJoin", [this]()
+			{
+				FConcertReplication_Join_Request JoinRequest;
+				JoinRequest.Streams.Add(ObjectReplicator->CreateStream());
+				
+				Cache.OnJoin(ClientId_One, JoinRequest);
+				
+				TestObjectInCache();
+			});
+
+			It("OnChangeStreams", [this]()
+			{
+				const FConcertReplicationStream JoinStream = ObjectReplicator->CreateStream();
+				FConcertReplication_Join_Request JoinRequest;
+				JoinRequest.Streams.Add(JoinStream);
+				
+				Cache.OnJoin(ClientId_One, JoinRequest);
+				Cache.OnChangeStreams(ClientId_One, { { FGuid::NewGuid(), SubobjectReplicator->TestObject } } , {});
+
+				TestObjectInCache();
+				TestSubobjectInCache();
+			});
+			
+			It("OnPreClientLeft", [this]()
+			{
+				const FConcertReplicationStream JoinStream = ObjectReplicator->CreateStream();
+				FConcertReplication_Join_Request JoinRequest;
+				JoinRequest.Streams.Add(JoinStream);
+				
+				Cache.OnJoin(ClientId_One, JoinRequest);
+				Cache.OnPreClientLeft(ClientId_One, { JoinStream });
+				
+				TestTrue(TEXT("IsEmpty"), Cache.IsEmpty());
+			});
+		});
+
+		Describe("Two clients", [this]()
+		{
+			It("Object retained when one client leaves", [this]()
+			{
+				const FConcertReplicationStream JoinStream = ObjectReplicator->CreateStream();
+				FConcertReplication_Join_Request JoinRequest;
+				JoinRequest.Streams.Add(JoinStream);
+				
+				Cache.OnJoin(ClientId_One, JoinRequest);
+				Cache.OnJoin(ClientId_Two, JoinRequest);
+				Cache.OnPreClientLeft(ClientId_One, { JoinStream });
+				
+				TestObjectInCache();
+			});
+
+			It("Object is removed when both clients leave", [this]()
+			{
+				const FConcertReplicationStream JoinStream = ObjectReplicator->CreateStream();
+				FConcertReplication_Join_Request JoinRequest;
+				JoinRequest.Streams.Add(JoinStream);
+				
+				Cache.OnJoin(ClientId_One, JoinRequest);
+				Cache.OnJoin(ClientId_Two, JoinRequest);
+				Cache.OnPreClientLeft(ClientId_One, { JoinStream });
+				Cache.OnPreClientLeft(ClientId_Two, { JoinStream });
+				
+				TestTrue(TEXT("IsEmpty"), Cache.IsEmpty());
+			});
+		});
+
+		It("IsObjectReferencedDirectly with ignored clients", [this]()
+		{
+			FConcertReplication_Join_Request JoinRequest;
+			JoinRequest.Streams.Add(ObjectReplicator->CreateStream());
+				
+			Cache.OnJoin(ClientId_One, JoinRequest);
+			Cache.OnJoin(ClientId_Two, JoinRequest);
+			Cache.OnJoin(ClientId_Three, JoinRequest);
+
+			const FGuid IgnoreOneTwo[] = { ClientId_One, ClientId_Two };
+			TestTrue(TEXT("IsObjectReferencedDirectly (ignore clients 1 and 2)"), Cache.IsObjectReferencedDirectly(ObjectReplicator->TestObject, IgnoreOneTwo));
+			const FGuid IgnoreAll[] = { ClientId_One, ClientId_Two, ClientId_Three };
+			TestTrue(TEXT("IsObjectReferencedDirectly (ignore all)"), !Cache.IsObjectReferencedDirectly(ObjectReplicator->TestObject, IgnoreAll));
+		});
+	}
+}
diff --git a/Engine/Plugins/Developer/Concert/ConcertSync/ConcertSyncTest/Source/ConcertSyncTest/Private/Replication/Util/Spec/ObjectTestReplicator.cpp b/Engine/Plugins/Developer/Concert/ConcertSync/ConcertSyncTest/Source/ConcertSyncTest/Private/Replication/Util/Spec/ObjectTestReplicator.cpp
--- a/Engine/Plugins/Developer/Concert/ConcertSync/ConcertSyncTest/Source/ConcertSyncTest/Private/Replication/Util/Spec/ObjectTestReplicator.cpp
+++ b/Engine/Plugins/Developer/Concert/ConcertSync/ConcertSyncTest/Source/ConcertSyncTest/Private/Replication/Util/Spec/ObjectTestReplicator.cpp
@@ -9,6 +9,13 @@
 
 namespace UE::ConcertSyncTests::Replication
 {
+	TSharedRef<FObjectTestReplicator> FObjectTestReplicator::CreateSubobjectReplicator() const
+	{
+		UTestReflectionObject* Subobject = NewObject<UTestReflectionObject>(TestObject);
+		TestObject->InstancedSubobject = Subobject;
+		return MakeShared<FObjectTestReplicator>(Subobject);
+	}
+
 	ConcertSyncClient::Replication::FJoinReplicatedSessionArgs FObjectTestReplicator::CreateSenderArgs(
 		FGuid SenderStreamId,
 		EConcertObjectReplicationMode ReplicationMode,
diff --git a/Engine/Plugins/Developer/Concert/ConcertSync/ConcertSyncTest/Source/ConcertSyncTest/Private/Replication/Util/Spec/ObjectTestReplicator.h b/Engine/Plugins/Developer/Concert/ConcertSync/ConcertSyncTest/Source/ConcertSyncTest/Private/Replication/Util/Spec/ObjectTestReplicator.h
--- a/Engine/Plugins/Developer/Concert/ConcertSync/ConcertSyncTest/Source/ConcertSyncTest/Private/Replication/Util/Spec/ObjectTestReplicator.h
+++ b/Engine/Plugins/Developer/Concert/ConcertSync/ConcertSyncTest/Source/ConcertSyncTest/Private/Replication/Util/Spec/ObjectTestReplicator.h
@@ -59,6 +59,12 @@ namespace UE::ConcertSyncTests::Replication
 		/** The object that will be transmitted. */
 		UTestReflectionObject* TestObject = NewObject<UTestReflectionObject>(GetTransientPackage());
 
+		FObjectTestReplicator() = default;
+		explicit FObjectTestReplicator(UTestReflectionObject* TestObject) : TestObject(TestObject) {}
+
+		/** Create a subobject of TestObject, assigns it to TestObject->InstancedSubobject, and returns a replicator for replicating it. */
+		TSharedRef<FObjectTestReplicator> CreateSubobjectReplicator() const;
+
 		/** Util for creating join args for replicating TestObject. */
 		ConcertSyncClient::Replication::FJoinReplicatedSessionArgs CreateSenderArgs(
 			FGuid SenderStreamId = FGuid::NewGuid(),
