Fix typo that caused hierarchy to be removed even though it is still in use. This caused dangling pointers and thus a crash.

diff --git a/Engine/Plugins/Developer/Concert/ConcertSync/ConcertSyncCore/Source/ConcertSyncCore/Private/Misc/ObjectPathHierarchy.cpp b/Engine/Plugins/Developer/Concert/ConcertSync/ConcertSyncCore/Source/ConcertSyncCore/Private/Misc/ObjectPathHierarchy.cpp
--- a/Engine/Plugins/Developer/Concert/ConcertSync/ConcertSyncCore/Source/ConcertSyncCore/Private/Misc/ObjectPathHierarchy.cpp
+++ b/Engine/Plugins/Developer/Concert/ConcertSync/ConcertSyncCore/Source/ConcertSyncCore/Private/Misc/ObjectPathHierarchy.cpp
@@ -74,6 +74,14 @@ namespace UE::ConcertSyncCore
 		return TreeNode && !(*TreeNode)->Children.IsEmpty();
 	}
 
+	bool FObjectPathHierarchy::IsAssetInHierarchy(const FSoftObjectPath& Object) const
+	{
+		return AssetNodes.ContainsByPredicate([&Object](const TUniquePtr<FTreeNode>& Node)
+		{
+			return Node->Data.Object == Object;
+		});
+	}
+
 	void FObjectPathHierarchy::AddObject(const FSoftObjectPath& ObjectPath)
 	{
 		if (FTreeNode** TreeNode = CachedNodes.Find(ObjectPath))
@@ -173,7 +181,7 @@ namespace UE::ConcertSyncCore
 			}
 			
 			// No more walking up the chain if we encounter an explicitly added object or a node that another object generated implicitly
-			bShouldStop = CurrentNode->Data.Type == EHierarchyObjectType::Explicit || CurrentNode->Children.Num() > 1;
+			bShouldStop = CurrentNode->Data.Type == EHierarchyObjectType::Explicit || CurrentNode->Children.Num() >= 1;
 		}
 	}
 
diff --git a/Engine/Plugins/Developer/Concert/ConcertSync/ConcertSyncCore/Source/ConcertSyncCore/Public/Misc/ObjectPathHierarchy.h b/Engine/Plugins/Developer/Concert/ConcertSync/ConcertSyncCore/Source/ConcertSyncCore/Public/Misc/ObjectPathHierarchy.h
--- a/Engine/Plugins/Developer/Concert/ConcertSync/ConcertSyncCore/Source/ConcertSyncCore/Public/Misc/ObjectPathHierarchy.h
+++ b/Engine/Plugins/Developer/Concert/ConcertSync/ConcertSyncCore/Source/ConcertSyncCore/Public/Misc/ObjectPathHierarchy.h
@@ -107,6 +107,8 @@ namespace UE::ConcertSyncCore
 		TOptional<EHierarchyObjectType> IsInHierarchy(const FSoftObjectPath& Object) const;
 		/** @return Whether Object has any subobjects in the hierarchy. */
 		bool HasChildren(const FSoftObjectPath& Object) const;
+		/** @return Whether Object is an asset, i.e. a top-level asset that has never have any parents. */
+		bool IsAssetInHierarchy(const FSoftObjectPath& Object) const;
 
 		/** Checks there is an hierarchy. */
 		bool IsEmpty() const { return AssetNodes.IsEmpty(); }
diff --git a/Engine/Plugins/Developer/Concert/ConcertSync/ConcertSyncTest/Source/ConcertSyncTest/Private/Misc/ObjectPathHierarchy.spec.cpp b/Engine/Plugins/Developer/Concert/ConcertSync/ConcertSyncTest/Source/ConcertSyncTest/Private/Misc/ObjectPathHierarchy.spec.cpp
--- a/Engine/Plugins/Developer/Concert/ConcertSync/ConcertSyncTest/Source/ConcertSyncTest/Private/Misc/ObjectPathHierarchy.spec.cpp
+++ b/Engine/Plugins/Developer/Concert/ConcertSync/ConcertSyncTest/Source/ConcertSyncTest/Private/Misc/ObjectPathHierarchy.spec.cpp
@@ -113,6 +113,22 @@ namespace UE::ConcertSyncTests
 				TestHasNoChildren(TEXT("/Game/Maps.Map:PersistentLevel.Sphere.StaticMeshComponent0"));
 			});
 			
+			It("IsAssetInHierarchy", [this]()
+			{
+				TestTrue(TEXT("IsAssetInHierarchy(Map)"), Hierarchy.IsAssetInHierarchy(FSoftObjectPath{ TEXT("/Game/Maps.Map")}));
+				
+				TestFalse(TEXT("IsAssetInHierarchy(PersistentLevel)"), Hierarchy.IsAssetInHierarchy(FSoftObjectPath{ TEXT("/Game/Maps.Map:PersistentLevel")}));
+				TestFalse(TEXT("IsAssetInHierarchy(Cube)"), Hierarchy.IsAssetInHierarchy(FSoftObjectPath{ TEXT("/Game/Maps.Map:PersistentLevel.Cube")}));
+				TestFalse(TEXT("IsAssetInHierarchy(StaticMeshComponent0)"), Hierarchy.IsAssetInHierarchy(FSoftObjectPath{ TEXT("/Game/Maps.Map:PersistentLevel.Cube.StaticMeshComponent0")}));
+				TestFalse(TEXT("IsAssetInHierarchy(StaticMeshComponent1)"), Hierarchy.IsAssetInHierarchy(FSoftObjectPath{ TEXT("/Game/Maps.Map:PersistentLevel.Cube.StaticMeshComponent1")}));
+				TestFalse(TEXT("IsAssetInHierarchy(Cube)"), Hierarchy.IsAssetInHierarchy(FSoftObjectPath{ TEXT("/Game/Maps.Map:PersistentLevel.Sphere")}));
+				TestFalse(TEXT("IsAssetInHierarchy(StaticMeshComponent0)"), Hierarchy.IsAssetInHierarchy(FSoftObjectPath{ TEXT("/Game/Maps.Map:PersistentLevel.Sphere.StaticMeshComponent0")}));
+				
+				TestFalse(TEXT("IsAssetInHierarchy(null)"), Hierarchy.IsAssetInHierarchy(FSoftObjectPath{ }));
+				TestFalse(TEXT("IsAssetInHierarchy(OtherMap)"), Hierarchy.IsAssetInHierarchy(FSoftObjectPath{ TEXT("/Game/Maps.OtherMap")}));
+				TestFalse(TEXT("IsAssetInHierarchy(OtherMap's StaticMeshComponent0)"), Hierarchy.IsAssetInHierarchy(FSoftObjectPath{ TEXT("/Game/Maps.OtherMap:PersistentLevel.Sphere.StaticMeshComponent0")}));
+			});
+			
 			It("TraverseTopToBottom (all)", [this]()
 			{
 				TArray<FSoftObjectPath> Visited;
@@ -363,6 +379,20 @@ namespace UE::ConcertSyncTests
 				TestHasNoChildren(TEXT("/Game/Maps.Map:PersistentLevel.Sphere.StaticMeshComponent0.Subobject"));
 			});
 
+			It("IsAssetInHierarchy", [this]()
+			{
+				TestTrue(TEXT("IsAssetInHierarchy(Map)"), Hierarchy.IsAssetInHierarchy(FSoftObjectPath{ TEXT("/Game/Maps.Map")}));
+				
+				TestFalse(TEXT("IsAssetInHierarchy(PersistentLevel)"), Hierarchy.IsAssetInHierarchy(FSoftObjectPath{ TEXT("/Game/Maps.Map:PersistentLevel")}));
+				TestFalse(TEXT("IsAssetInHierarchy(Cube)"), Hierarchy.IsAssetInHierarchy(FSoftObjectPath{ TEXT("/Game/Maps.Map:PersistentLevel.Sphere")}));
+				TestFalse(TEXT("IsAssetInHierarchy(StaticMeshComponent0)"), Hierarchy.IsAssetInHierarchy(FSoftObjectPath{ TEXT("/Game/Maps.Map:PersistentLevel.Sphere.StaticMeshComponent0")}));
+				TestFalse(TEXT("IsAssetInHierarchy(StaticMeshComponent0)"), Hierarchy.IsAssetInHierarchy(FSoftObjectPath{ TEXT("/Game/Maps.Map:PersistentLevel.Sphere.StaticMeshComponent0.Subobject")}));
+				
+				TestFalse(TEXT("IsAssetInHierarchy(null)"), Hierarchy.IsAssetInHierarchy(FSoftObjectPath{ }));
+				TestFalse(TEXT("IsAssetInHierarchy(OtherMap)"), Hierarchy.IsAssetInHierarchy(FSoftObjectPath{ TEXT("/Game/Maps.OtherMap")}));
+				TestFalse(TEXT("IsAssetInHierarchy(OtherMap's StaticMeshComponent0)"), Hierarchy.IsAssetInHierarchy(FSoftObjectPath{ TEXT("/Game/Maps.OtherMap:PersistentLevel.Sphere.StaticMeshComponent0")}));
+			});
+
 			It("TraverseTopToBottom (all)", [this]()
 			{
 				TArray<FSoftObjectPath> Visited;
@@ -522,5 +552,22 @@ namespace UE::ConcertSyncTests
 				Hierarchy.RemoveObject(FSoftObjectPath{ TEXT("/Engine/Transient.Root") });
 			});
 		});
+
+		It("Retain hierarchy after removal", [this]()
+		{
+			Hierarchy.AddObject(FSoftObjectPath{ TEXT("/Game/Maps.Map:PersistentLevel.Cube.Foo") });
+			Hierarchy.AddObject(FSoftObjectPath{ TEXT("/Game/Maps.Map:PersistentLevel.Cube.Bar") });
+			// We'll proceed to test that the rest of the hierarchy remains intact.
+			Hierarchy.RemoveObject(FSoftObjectPath{ TEXT("/Game/Maps.Map:PersistentLevel.Cube.Foo")});
+
+			// There used to be a bug where FObjectPathHierarchy::AssetNodes would be emptied incorrectly and FObjectPathHierarchy::CachedNodes left dangling. Check that case specifically with IsAssetInHierarchy:
+			TestTrue(TEXT("IsAsset(Map)"), Hierarchy.IsAssetInHierarchy(FSoftObjectPath{ TEXT("/Game/Maps.Map")}));
+			
+			TestTrue(TEXT("IsInHierarchy(Map)"), Hierarchy.IsInHierarchy(FSoftObjectPath{ TEXT("/Game/Maps.Map")}).IsSet());
+			TestTrue(TEXT("IsInHierarchy(Cube)"), Hierarchy.IsInHierarchy(FSoftObjectPath{ TEXT("/Game/Maps.Map:PersistentLevel")}).IsSet());
+			TestTrue(TEXT("IsInHierarchy(Cube)"), Hierarchy.IsInHierarchy(FSoftObjectPath{ TEXT("/Game/Maps.Map:PersistentLevel.Cube")}).IsSet());
+			TestTrue(TEXT("IsInHierarchy(Bar)"), Hierarchy.IsInHierarchy(FSoftObjectPath{ TEXT("/Game/Maps.Map:PersistentLevel.Cube.Bar")}).IsSet());
+			TestTrue(TEXT("HasChildren(Cube)"), Hierarchy.HasChildren(FSoftObjectPath{ TEXT("/Game/Maps.Map:PersistentLevel.Cube")}));
+		});
 	}
 }
