Add ability to mute & unmute objects

diff --git a/Engine/Plugins/Developer/Concert/ConcertSync/ConcertSyncClient/Source/ConcertSyncClient/Private/ConcertSyncClient.cpp b/Engine/Plugins/Developer/Concert/ConcertSync/ConcertSyncClient/Source/ConcertSyncClient/Private/ConcertSyncClient.cpp
--- a/Engine/Plugins/Developer/Concert/ConcertSync/ConcertSyncClient/Source/ConcertSyncClient/Private/ConcertSyncClient.cpp
+++ b/Engine/Plugins/Developer/Concert/ConcertSync/ConcertSyncClient/Source/ConcertSyncClient/Private/ConcertSyncClient.cpp
@@ -189,7 +189,7 @@ void FConcertSyncClient::RegisterConcertSyncHandlers(TSharedRef<IConcertClientSe
 		// Create Replication Manager
 		if (EnumHasAnyFlags(LiveSession->GetSessionFlags(), EConcertSyncSessionFlags::EnableReplication) && ensure(Bridges.ReplicationBridge))
 		{
-			ReplicationManager = MakeUnique<UE::ConcertSyncClient::Replication::FReplicationManager>(InSession, *Bridges.ReplicationBridge);
+			ReplicationManager = MakeUnique<UE::ConcertSyncClient::Replication::FReplicationManager>(InSession, *Bridges.ReplicationBridge, LiveSession->GetSessionFlags());
 			ReplicationManager->StartAcceptingJoinRequests();
 		}
 
diff --git a/Engine/Plugins/Developer/Concert/ConcertSync/ConcertSyncClient/Source/ConcertSyncClient/Private/Replication/IConcertClientReplicationManager.cpp b/Engine/Plugins/Developer/Concert/ConcertSync/ConcertSyncClient/Source/ConcertSyncClient/Private/Replication/IConcertClientReplicationManager.cpp
--- a/Engine/Plugins/Developer/Concert/ConcertSync/ConcertSyncClient/Source/ConcertSyncClient/Private/Replication/IConcertClientReplicationManager.cpp
+++ b/Engine/Plugins/Developer/Concert/ConcertSync/ConcertSyncClient/Source/ConcertSyncClient/Private/Replication/IConcertClientReplicationManager.cpp
@@ -13,19 +13,18 @@ namespace UE::ConcertSyncClient::Replication::Private
 	{
 		for (const FSoftObjectPath& ObjectPath : Objects)
 		{
-			Manager.ForEachRegisteredStream([&InCallback, &ObjectPath](const FConcertReplicationStream& StreamDescription)
+			bool bWasFound = false;
+			Manager.ForEachRegisteredStream([&InCallback, &ObjectPath, &bWasFound](const FConcertReplicationStream& StreamDescription)
 			{
 				const bool bStreamContainsObject = StreamDescription.BaseDescription.ReplicationMap.ReplicatedObjects.Contains(ObjectPath);
 				if (bStreamContainsObject)
 				{
+					bWasFound = true;
 					InCallback(ObjectPath, StreamDescription.BaseDescription.Identifier);
 				}
-				else
-				{
-					UE_LOG(LogConcert, Warning, TEXT("Object %s is not a valid argument because it is not contained in any stream."), *ObjectPath.ToString());
-				}
 				return EBreakBehavior::Continue;
 			});
+			UE_CLOG(!bWasFound, LogConcert, Warning, TEXT("Object %s is not a valid argument because it is not contained in any stream."), *ObjectPath.ToString());
 		}
 	}
 }
@@ -42,7 +41,7 @@ TArray<FConcertReplicationStream> IConcertClientReplicationManager::GetRegistere
 	return Result;
 }
 
-TFuture<FConcertReplication_ChangeAuthority_Response> IConcertClientReplicationManager::TakeAuthorityOver(TArrayView<const FSoftObjectPath> Objects)
+TFuture<FConcertReplication_ChangeAuthority_Response> IConcertClientReplicationManager::TakeAuthorityOver(TConstArrayView<FSoftObjectPath> Objects)
 {
 	using namespace UE::ConcertSyncClient::Replication;
 	
@@ -73,7 +72,7 @@ TFuture<FConcertReplication_ChangeAuthority_Response> IConcertClientReplicationM
 	return RequestAuthorityChange(MoveTemp(Request));
 }
 
-TFuture<FConcertReplication_ChangeAuthority_Response> IConcertClientReplicationManager::ReleaseAuthorityOf(TArrayView<const FSoftObjectPath> Objects)
+TFuture<FConcertReplication_ChangeAuthority_Response> IConcertClientReplicationManager::ReleaseAuthorityOf(TConstArrayView<FSoftObjectPath> Objects)
 {
 	using namespace UE::ConcertSyncClient::Replication;
 	
@@ -132,3 +131,28 @@ TSet<FConcertObjectInStreamID> IConcertClientReplicationManager::GetSyncControll
 	});
 	return Result;
 }
+
+TFuture<FConcertReplication_ChangeMuteState_Response> IConcertClientReplicationManager::MuteObjects(TConstArrayView<FSoftObjectPath> Objects, EConcertReplicationMuteOption Flags)
+{
+	FConcertReplication_ChangeMuteState_Request Request;
+	for (const FSoftObjectPath& Object : Objects)
+	{
+		Request.ObjectsToMute.Add(Object, { Flags });
+	}
+	return ChangeMuteState(Request);
+}
+
+TFuture<FConcertReplication_ChangeMuteState_Response> IConcertClientReplicationManager::UnmuteObjects(TSet<FSoftObjectPath> Objects, EConcertReplicationMuteOption Flags)
+{
+	FConcertReplication_ChangeMuteState_Request Request;
+	for (const FSoftObjectPath& Object : Objects)
+	{
+		Request.ObjectsToUnmute.Add(Object, { Flags });
+	}
+	return ChangeMuteState(Request);
+}
+
+TFuture<FConcertReplication_QueryMuteState_Response> IConcertClientReplicationManager::QueryMuteState(TSet<FSoftObjectPath> Objects)
+{
+	return QueryMuteState(FConcertReplication_QueryMuteState_Request{ MoveTemp(Objects) });
+}
diff --git a/Engine/Plugins/Developer/Concert/ConcertSync/ConcertSyncClient/Source/ConcertSyncClient/Private/Replication/Manager/ReplicationManager.cpp b/Engine/Plugins/Developer/Concert/ConcertSync/ConcertSyncClient/Source/ConcertSyncClient/Private/Replication/Manager/ReplicationManager.cpp
--- a/Engine/Plugins/Developer/Concert/ConcertSync/ConcertSyncClient/Source/ConcertSyncClient/Private/Replication/Manager/ReplicationManager.cpp
+++ b/Engine/Plugins/Developer/Concert/ConcertSync/ConcertSyncClient/Source/ConcertSyncClient/Private/Replication/Manager/ReplicationManager.cpp
@@ -10,9 +10,14 @@
 
 namespace UE::ConcertSyncClient::Replication
 {
-	FReplicationManager::FReplicationManager(TSharedRef<IConcertClientSession> InLiveSession, IConcertClientReplicationBridge& InBridge)
+	FReplicationManager::FReplicationManager(
+		TSharedRef<IConcertClientSession> InLiveSession,
+		IConcertClientReplicationBridge& InBridge,
+		EConcertSyncSessionFlags SessionFlags
+		)
 		: Session(MoveTemp(InLiveSession))
 		, Bridge(InBridge)
+		, SessionFlags(SessionFlags)
 	{}
 
 	FReplicationManager::~FReplicationManager()
@@ -21,7 +26,7 @@ namespace UE::ConcertSyncClient::Replication
 	void FReplicationManager::StartAcceptingJoinRequests()
 	{
 		checkSlow(!CurrentState.IsValid());
-		CurrentState = MakeShared<FReplicationManagerState_Disconnected>(Session, Bridge, *this);
+		CurrentState = MakeShared<FReplicationManagerState_Disconnected>(Session, Bridge, *this, SessionFlags);
 	}
 
 	TFuture<FJoinReplicatedSessionResult> FReplicationManager::JoinReplicationSession(FJoinReplicatedSessionArgs Args)
@@ -121,6 +126,20 @@ namespace UE::ConcertSyncClient::Replication
 			&& CurrentState->HasSyncControl(Object);
 	}
 
+	TFuture<FConcertReplication_ChangeMuteState_Response> FReplicationManager::ChangeMuteState(FConcertReplication_ChangeMuteState_Request Request)
+	{
+		return ensureMsgf(CurrentState, TEXT("StartAcceptingJoinRequests should have been called at this point."))
+			? CurrentState->ChangeMuteState(MoveTemp(Request))
+			: MakeFulfilledPromise<FConcertReplication_ChangeMuteState_Response>().GetFuture(); 
+	}
+
+	TFuture<FConcertReplication_QueryMuteState_Response> FReplicationManager::QueryMuteState(FConcertReplication_QueryMuteState_Request Request)
+	{
+		return ensureMsgf(CurrentState, TEXT("StartAcceptingJoinRequests should have been called at this point."))
+			? CurrentState->QueryMuteState(MoveTemp(Request))
+			: MakeFulfilledPromise<FConcertReplication_QueryMuteState_Response>().GetFuture(); 
+	}
+
 	IConcertClientReplicationManager::FOnPreStreamsChanged& FReplicationManager::OnPreStreamsChanged()
 	{
 		// Check() to avoid returning some dummy static variable
diff --git a/Engine/Plugins/Developer/Concert/ConcertSync/ConcertSyncClient/Source/ConcertSyncClient/Private/Replication/Manager/ReplicationManager.h b/Engine/Plugins/Developer/Concert/ConcertSync/ConcertSyncClient/Source/ConcertSyncClient/Private/Replication/Manager/ReplicationManager.h
--- a/Engine/Plugins/Developer/Concert/ConcertSync/ConcertSyncClient/Source/ConcertSyncClient/Private/Replication/Manager/ReplicationManager.h
+++ b/Engine/Plugins/Developer/Concert/ConcertSync/ConcertSyncClient/Source/ConcertSyncClient/Private/Replication/Manager/ReplicationManager.h
@@ -2,6 +2,7 @@
 
 #pragma once
 
+#include "ConcertSyncSessionFlags.h"
 #include "Replication/IConcertClientReplicationManager.h"
 #include "Replication/Messages/Handshake.h"
 #include "Templates/SharedPointer.h"
@@ -18,7 +19,11 @@ namespace UE::ConcertSyncClient::Replication
 		friend FReplicationManagerState; 
 	public:
 		
-		FReplicationManager(TSharedRef<IConcertClientSession> InLiveSession, IConcertClientReplicationBridge& InBridge UE_LIFETIMEBOUND);
+		FReplicationManager(
+			TSharedRef<IConcertClientSession> InLiveSession,
+			IConcertClientReplicationBridge& InBridge UE_LIFETIMEBOUND,
+			EConcertSyncSessionFlags SessionFlags
+			);
 		virtual ~FReplicationManager() override;
 
 		/** Starts accepting join requests. Must be called separately from constructor because of TSharedFromThis asserting if SharedThis is called in constructor. */
@@ -38,6 +43,8 @@ namespace UE::ConcertSyncClient::Replication
 		virtual ESyncControlEnumerationResult ForEachSyncControlledObject(TFunctionRef<EBreakBehavior(const FConcertObjectInStreamID& Object)> Callback) const override;
 		virtual uint32 NumSyncControlledObjects() const override;
 		virtual bool HasSyncControl(const FConcertObjectInStreamID& Object) const override;
+		virtual TFuture<FConcertReplication_ChangeMuteState_Response> ChangeMuteState(FConcertReplication_ChangeMuteState_Request Request) override;
+		virtual TFuture<FConcertReplication_QueryMuteState_Response> QueryMuteState(FConcertReplication_QueryMuteState_Request Request) override;
 		virtual FOnPreStreamsChanged& OnPreStreamsChanged() override;
 		virtual FOnPostStreamsChanged& OnPostStreamsChanged() override;
 		virtual FOnPreAuthorityChanged& OnPreAuthorityChanged() override;
@@ -52,6 +59,8 @@ namespace UE::ConcertSyncClient::Replication
 		TSharedRef<IConcertClientSession> Session;
 		/** The replication bridge is responsible for applying received data and generating data to send. */
 		IConcertClientReplicationBridge& Bridge;
+		/** These flags are passed along to all the states. */
+		const EConcertSyncSessionFlags SessionFlags;
 
 		/** The current state this manager is in, e.g. waiting for connection request, connecting, connected, etc. */
 		TSharedPtr<FReplicationManagerState> CurrentState;
diff --git a/Engine/Plugins/Developer/Concert/ConcertSync/ConcertSyncClient/Source/ConcertSyncClient/Private/Replication/Manager/ReplicationManagerState.h b/Engine/Plugins/Developer/Concert/ConcertSync/ConcertSyncClient/Source/ConcertSyncClient/Private/Replication/Manager/ReplicationManagerState.h
--- a/Engine/Plugins/Developer/Concert/ConcertSync/ConcertSyncClient/Source/ConcertSyncClient/Private/Replication/Manager/ReplicationManagerState.h
+++ b/Engine/Plugins/Developer/Concert/ConcertSync/ConcertSyncClient/Source/ConcertSyncClient/Private/Replication/Manager/ReplicationManagerState.h
@@ -4,6 +4,7 @@
 
 #include "ReplicationManagerUtils.h"
 #include "Replication/IConcertClientReplicationManager.h"
+#include "Replication/Messages/Muting.h"
 
 namespace UE::ConcertSyncClient::Replication
 {
@@ -31,6 +32,8 @@ namespace UE::ConcertSyncClient::Replication
 		virtual ESyncControlEnumerationResult ForEachSyncControlledObject(TFunctionRef<EBreakBehavior(const FConcertObjectInStreamID& Object)> Callback) const override { return ESyncControlEnumerationResult::NoneAvailable; }
 		virtual uint32 NumSyncControlledObjects() const override { return 0; }
 		virtual bool HasSyncControl(const FConcertObjectInStreamID& Object) const override { return false; }
+		virtual TFuture<FConcertReplication_ChangeMuteState_Response> ChangeMuteState(FConcertReplication_ChangeMuteState_Request) override { return MakeFulfilledPromise<FConcertReplication_ChangeMuteState_Response>(FConcertReplication_ChangeMuteState_Response{ EConcertReplicationMuteErrorCode::Rejected }).GetFuture(); };
+		virtual TFuture<FConcertReplication_QueryMuteState_Response> QueryMuteState(FConcertReplication_QueryMuteState_Request Request) override { return MakeFulfilledPromise<FConcertReplication_QueryMuteState_Response>().GetFuture(); }
 		virtual FOnPreStreamsChanged& OnPreStreamsChanged() override { return OnPreStreamsChangedDelegate; } 
 		virtual FOnPostStreamsChanged& OnPostStreamsChanged() override { return OnPostStreamsChangedDelegate; }
 		virtual FOnPreAuthorityChanged& OnPreAuthorityChanged() override { return OnPreAuthorityChangedDelegate; }
diff --git a/Engine/Plugins/Developer/Concert/ConcertSync/ConcertSyncClient/Source/ConcertSyncClient/Private/Replication/Manager/ReplicationManagerState_Connected.cpp b/Engine/Plugins/Developer/Concert/ConcertSync/ConcertSyncClient/Source/ConcertSyncClient/Private/Replication/Manager/ReplicationManagerState_Connected.cpp
--- a/Engine/Plugins/Developer/Concert/ConcertSync/ConcertSyncClient/Source/ConcertSyncClient/Private/Replication/Manager/ReplicationManagerState_Connected.cpp
+++ b/Engine/Plugins/Developer/Concert/ConcertSync/ConcertSyncClient/Source/ConcertSyncClient/Private/Replication/Manager/ReplicationManagerState_Connected.cpp
@@ -42,6 +42,11 @@ namespace UE::ConcertSyncClient::Replication
 		false,
 		TEXT("Whether the client should pretend that authority change requests were rejected.")
 		);
+	TAutoConsoleVariable<bool> CVarSimulateMuteRequestRejection(
+		TEXT("Concert.Replication.SimulateMuteRejection"),
+		false,
+		TEXT("Whether the client should pretend that mute change requests were rejected.")
+		);
 
 	TAutoConsoleVariable<bool> CVarLogStreamRequestsAndResponsesOnClient(
 		TEXT("Concert.Replication.LogStreamRequestsAndResponsesOnClient"),
@@ -53,6 +58,11 @@ namespace UE::ConcertSyncClient::Replication
 		false,
 		TEXT("Whether to log changes to authority.")
 		);
+	TAutoConsoleVariable<bool> CVarLogMuteRequestsAndResponsesOnClient(
+		TEXT("Concert.Replication.LogMuteRequestsAndResponsesOnClient"),
+		false,
+		TEXT("Whether to log changes to the mute state.")
+		);
 
 	namespace Private
 	{
@@ -72,12 +82,14 @@ namespace UE::ConcertSyncClient::Replication
 		TSharedRef<IConcertClientSession> InLiveSession,
 		IConcertClientReplicationBridge& ReplicationBridge,
 		FReplicationManager& Owner,
+		EConcertSyncSessionFlags SessionFlags,
 		TArray<FConcertReplicationStream> InitialStreams,
 		const FConcertReplication_ChangeSyncControl& InitialSyncControl
 		)
 		: FReplicationManagerState(Owner)
 		, LiveSession(InLiveSession)
 		, ReplicationBridge(ReplicationBridge)
+		, SessionFlags(SessionFlags)
 		, RegisteredStreams(MoveTemp(InitialStreams))
 		// TODO DP: Use config to determine which replication format to use
 		, ReplicationFormat(MakeUnique<ConcertSyncCore::FFullObjectFormat>())
@@ -121,7 +133,7 @@ namespace UE::ConcertSyncClient::Replication
 	void FReplicationManagerState_Connected::LeaveReplicationSession()
 	{
 		LiveSession->SendCustomEvent(FConcertReplication_LeaveEvent{}, LiveSession->GetSessionServerEndpointId(), EConcertMessageFlags::ReliableOrdered);
-		ChangeState(MakeShared<FReplicationManagerState_Disconnected>(LiveSession, ReplicationBridge, GetOwner()));
+		ChangeState(MakeShared<FReplicationManagerState_Disconnected>(LiveSession, ReplicationBridge, GetOwner(), SessionFlags));
 	}
 
 	IConcertClientReplicationManager::EStreamEnumerationResult FReplicationManagerState_Connected::ForEachRegisteredStream(
@@ -266,6 +278,43 @@ namespace UE::ConcertSyncClient::Replication
 			: ESyncControlEnumerationResult::NoneAvailable;
 	}
 
+	TFuture<FConcertReplication_ChangeMuteState_Response> FReplicationManagerState_Connected::ChangeMuteState(FConcertReplication_ChangeMuteState_Request Request)
+	{
+		if (!EnumHasAnyFlags(SessionFlags, EConcertSyncSessionFlags::ShouldAllowGlobalMuting) || CVarSimulateMuteRequestRejection.GetValueOnAnyThread())
+		{
+			return MakeFulfilledPromise<FConcertReplication_ChangeMuteState_Response>(FConcertReplication_ChangeMuteState_Response{ EConcertReplicationMuteErrorCode::Rejected }).GetFuture();
+		}
+
+		FLocalSyncControl::FPredictedObjectRemoval PredictedChanges = SyncControl.PredictAndApplyMuteChanges(Request);
+		
+		Private::LogNetworkMessage(CVarLogMuteRequestsAndResponsesOnClient, Request);
+		return LiveSession->SendCustomRequest<FConcertReplication_ChangeMuteState_Request, FConcertReplication_ChangeMuteState_Response>(Request, LiveSession->GetSessionServerEndpointId())
+			.Next([WeakThis = TWeakPtr<FReplicationManagerState_Connected>(SharedThis(this)), PredictedChanges = MoveTemp(PredictedChanges)](FConcertReplication_ChangeMuteState_Response&& Response)
+			{
+				Private::LogNetworkMessage(CVarLogMuteRequestsAndResponsesOnClient, Response);
+				
+				if (const TSharedPtr<FReplicationManagerState_Connected> This = WeakThis.Pin())
+				{
+					This->SyncControl.ApplyOrRevertMuteResponse(PredictedChanges, Response);
+				}
+				
+				return Response;
+			});
+	}
+
+	TFuture<FConcertReplication_QueryMuteState_Response> FReplicationManagerState_Connected::QueryMuteState(FConcertReplication_QueryMuteState_Request Request)
+	{
+		if (!EnumHasAnyFlags(SessionFlags, EConcertSyncSessionFlags::ShouldAllowGlobalMuting))
+		{
+			return MakeFulfilledPromise<FConcertReplication_QueryMuteState_Response>().GetFuture();
+		}
+		
+		return LiveSession->SendCustomRequest<FConcertReplication_QueryMuteState_Request, FConcertReplication_QueryMuteState_Response>(
+			Request,
+			LiveSession->GetSessionServerEndpointId()
+			);
+	}
+
 	void FReplicationManagerState_Connected::OnEnterState()
 	{
 		LiveSession->OnTick().AddSP(this, &FReplicationManagerState_Connected::Tick);
diff --git a/Engine/Plugins/Developer/Concert/ConcertSync/ConcertSyncClient/Source/ConcertSyncClient/Private/Replication/Manager/ReplicationManagerState_Connected.h b/Engine/Plugins/Developer/Concert/ConcertSync/ConcertSyncClient/Source/ConcertSyncClient/Private/Replication/Manager/ReplicationManagerState_Connected.h
--- a/Engine/Plugins/Developer/Concert/ConcertSync/ConcertSyncClient/Source/ConcertSyncClient/Private/Replication/Manager/ReplicationManagerState_Connected.h
+++ b/Engine/Plugins/Developer/Concert/ConcertSync/ConcertSyncClient/Source/ConcertSyncClient/Private/Replication/Manager/ReplicationManagerState_Connected.h
@@ -2,6 +2,7 @@
 
 #pragma once
 
+#include "ConcertSyncSessionFlags.h"
 #include "ReplicationManagerState.h"
 #include "Replication/Formats/IObjectReplicationFormat.h"
 #include "Replication/Misc/LocalSyncControl.h"
@@ -33,6 +34,7 @@ namespace UE::ConcertSyncClient::Replication
 			TSharedRef<IConcertClientSession> InLiveSession,
 			IConcertClientReplicationBridge& ReplicationBridge UE_LIFETIMEBOUND,
 			FReplicationManager& Owner UE_LIFETIMEBOUND,
+			EConcertSyncSessionFlags SessionFlags,
 			TArray<FConcertReplicationStream> InitialStreams,
 			const FConcertReplication_ChangeSyncControl& InitialSyncControl
 			);
@@ -52,6 +54,8 @@ namespace UE::ConcertSyncClient::Replication
 		virtual ESyncControlEnumerationResult ForEachSyncControlledObject(TFunctionRef<EBreakBehavior(const FConcertObjectInStreamID& Object)> Callback) const override;
 		virtual uint32 NumSyncControlledObjects() const override { return SyncControl.Num(); }
 		virtual bool HasSyncControl(const FConcertObjectInStreamID& Object) const override { return SyncControl.IsObjectAllowed(Object); }
+		virtual TFuture<FConcertReplication_ChangeMuteState_Response> ChangeMuteState(FConcertReplication_ChangeMuteState_Request) override;
+		virtual TFuture<FConcertReplication_QueryMuteState_Response> QueryMuteState(FConcertReplication_QueryMuteState_Request Request) override;
 		//~ End IConcertClientReplicationManager Interface
 
 	private:
@@ -60,6 +64,8 @@ namespace UE::ConcertSyncClient::Replication
 		const TSharedRef<IConcertClientSession> LiveSession;
 		/** Passed to FReplicationManagerState_Disconnected */
 		IConcertClientReplicationBridge& ReplicationBridge;
+		/** Passed to FReplicationManagerState_Disconnected and used to determine whether certain operations are supported by the server. */
+		const EConcertSyncSessionFlags SessionFlags;
 		/** The streams this client has registered with the server. */
 		TArray<FConcertReplicationStream> RegisteredStreams;
 		
diff --git a/Engine/Plugins/Developer/Concert/ConcertSync/ConcertSyncClient/Source/ConcertSyncClient/Private/Replication/Manager/ReplicationManagerState_Disconnected.cpp b/Engine/Plugins/Developer/Concert/ConcertSync/ConcertSyncClient/Source/ConcertSyncClient/Private/Replication/Manager/ReplicationManagerState_Disconnected.cpp
--- a/Engine/Plugins/Developer/Concert/ConcertSync/ConcertSyncClient/Source/ConcertSyncClient/Private/Replication/Manager/ReplicationManagerState_Disconnected.cpp
+++ b/Engine/Plugins/Developer/Concert/ConcertSync/ConcertSyncClient/Source/ConcertSyncClient/Private/Replication/Manager/ReplicationManagerState_Disconnected.cpp
@@ -18,11 +18,13 @@ namespace UE::ConcertSyncClient::Replication
 	FReplicationManagerState_Disconnected::FReplicationManagerState_Disconnected(
 		TSharedRef<IConcertClientSession> LiveSession,
 		IConcertClientReplicationBridge& ReplicationBridge,
-		FReplicationManager& Owner
+		FReplicationManager& Owner,
+		EConcertSyncSessionFlags SessionFlags
 		)
 		: FReplicationManagerState(Owner)
 		, LiveSession(MoveTemp(LiveSession))
 		, ReplicationBridge(ReplicationBridge)
+		, SessionFlags(SessionFlags)
 	{}
 
 	TFuture<FJoinReplicatedSessionResult> FReplicationManagerState_Disconnected::JoinReplicationSession(FJoinReplicatedSessionArgs Args)
@@ -34,7 +36,7 @@ namespace UE::ConcertSyncClient::Replication
 		
 		TPromise<FJoinReplicatedSessionResult> JoinPromise;
 		TFuture<FJoinReplicatedSessionResult> JoinFuture = JoinPromise.GetFuture();
-		ChangeState(MakeShared<FReplicationManagerState_Handshaking>(MoveTemp(Args), MoveTemp(JoinPromise), LiveSession, ReplicationBridge, GetOwner()));
+		ChangeState(MakeShared<FReplicationManagerState_Handshaking>(MoveTemp(Args), MoveTemp(JoinPromise), LiveSession, ReplicationBridge, GetOwner(), SessionFlags));
 		return JoinFuture;
 	}
 
diff --git a/Engine/Plugins/Developer/Concert/ConcertSync/ConcertSyncClient/Source/ConcertSyncClient/Private/Replication/Manager/ReplicationManagerState_Disconnected.h b/Engine/Plugins/Developer/Concert/ConcertSync/ConcertSyncClient/Source/ConcertSyncClient/Private/Replication/Manager/ReplicationManagerState_Disconnected.h
--- a/Engine/Plugins/Developer/Concert/ConcertSync/ConcertSyncClient/Source/ConcertSyncClient/Private/Replication/Manager/ReplicationManagerState_Disconnected.h
+++ b/Engine/Plugins/Developer/Concert/ConcertSync/ConcertSyncClient/Source/ConcertSyncClient/Private/Replication/Manager/ReplicationManagerState_Disconnected.h
@@ -2,8 +2,8 @@
 
 #pragma once
 
+#include "ConcertSyncSessionFlags.h"
 #include "ReplicationManagerState.h"
-#include "ReplicationManagerUtils.h"
 #include "Replication/Processing/ClientReplicationDataCollector.h"
 
 class IConcertClientSession;
@@ -18,7 +18,8 @@ namespace UE::ConcertSyncClient::Replication
 		FReplicationManagerState_Disconnected(
 			TSharedRef<IConcertClientSession> LiveSession,
 			IConcertClientReplicationBridge& ReplicationBridge UE_LIFETIMEBOUND,
-			FReplicationManager& Owner UE_LIFETIMEBOUND
+			FReplicationManager& Owner UE_LIFETIMEBOUND,
+			EConcertSyncSessionFlags SessionFlags
 			);
 
 		//~ Begin IConcertClientReplicationManager Interface
@@ -34,5 +35,7 @@ namespace UE::ConcertSyncClient::Replication
 		TSharedRef<IConcertClientSession> LiveSession;
 		/** Passed to FReplicationManagerState_Handshaking */
 		IConcertClientReplicationBridge& ReplicationBridge;
+		/** Passed to FReplicationManagerState_Handshaking */
+		const EConcertSyncSessionFlags SessionFlags;
 	};
 }
\ No newline at end of file
diff --git a/Engine/Plugins/Developer/Concert/ConcertSync/ConcertSyncClient/Source/ConcertSyncClient/Private/Replication/Manager/ReplicationManagerState_Handshaking.cpp b/Engine/Plugins/Developer/Concert/ConcertSync/ConcertSyncClient/Source/ConcertSyncClient/Private/Replication/Manager/ReplicationManagerState_Handshaking.cpp
--- a/Engine/Plugins/Developer/Concert/ConcertSync/ConcertSyncClient/Source/ConcertSyncClient/Private/Replication/Manager/ReplicationManagerState_Handshaking.cpp
+++ b/Engine/Plugins/Developer/Concert/ConcertSync/ConcertSyncClient/Source/ConcertSyncClient/Private/Replication/Manager/ReplicationManagerState_Handshaking.cpp
@@ -16,13 +16,15 @@ namespace UE::ConcertSyncClient::Replication
 		TPromise<FJoinReplicatedSessionResult> JoinSessionPromise,
 		TSharedRef<IConcertClientSession> LiveSession,
 		IConcertClientReplicationBridge& ReplicationBridge,
-		FReplicationManager& Owner
+		FReplicationManager& Owner,
+		EConcertSyncSessionFlags SessionFlags
 		)
 		: FReplicationManagerState(Owner)
 		, RequestArgs(MoveTemp(RequestArgs))
 		, JoinSessionPromise(MoveTemp(JoinSessionPromise))
 		, LiveSession(MoveTemp(LiveSession))
 		, ReplicationBridge(ReplicationBridge)
+		, SessionFlags(SessionFlags)
 	{}
 
 	FReplicationManagerState_Handshaking::~FReplicationManagerState_Handshaking()
@@ -62,7 +64,9 @@ namespace UE::ConcertSyncClient::Replication
 				if (Response.JoinErrorCode == EJoinReplicationErrorCode::Success)
 				{
 					This->ChangeState(
-						MakeShared<FReplicationManagerState_Connected>(This->LiveSession, This->ReplicationBridge, This->GetOwner(), MoveTemp(This->RequestArgs.Streams), Response.SyncControl)
+						MakeShared<FReplicationManagerState_Connected>(
+							This->LiveSession, This->ReplicationBridge, This->GetOwner(), This->SessionFlags, MoveTemp(This->RequestArgs.Streams), Response.SyncControl
+							)
 						);
 				}
 				else
@@ -78,7 +82,7 @@ namespace UE::ConcertSyncClient::Replication
 	void FReplicationManagerState_Handshaking::ReturnToDisconnectedState()
 	{
 		ChangeState(
-			MakeShared<FReplicationManagerState_Disconnected>(LiveSession, ReplicationBridge, GetOwner())
+			MakeShared<FReplicationManagerState_Disconnected>(LiveSession, ReplicationBridge, GetOwner(), SessionFlags)
 			);
 	}
 }
diff --git a/Engine/Plugins/Developer/Concert/ConcertSync/ConcertSyncClient/Source/ConcertSyncClient/Private/Replication/Manager/ReplicationManagerState_Handshaking.h b/Engine/Plugins/Developer/Concert/ConcertSync/ConcertSyncClient/Source/ConcertSyncClient/Private/Replication/Manager/ReplicationManagerState_Handshaking.h
--- a/Engine/Plugins/Developer/Concert/ConcertSync/ConcertSyncClient/Source/ConcertSyncClient/Private/Replication/Manager/ReplicationManagerState_Handshaking.h
+++ b/Engine/Plugins/Developer/Concert/ConcertSync/ConcertSyncClient/Source/ConcertSyncClient/Private/Replication/Manager/ReplicationManagerState_Handshaking.h
@@ -2,10 +2,8 @@
 
 #pragma once
 
-#include "IConcertSessionHandler.h"
+#include "ConcertSyncSessionFlags.h"
 #include "ReplicationManagerState.h"
-#include "ReplicationManagerUtils.h"
-#include "Replication/Messages/Handshake.h"
 #include "Replication/Processing/ClientReplicationDataCollector.h"
 
 class IConcertClientSession;
@@ -22,7 +20,8 @@ namespace UE::ConcertSyncClient::Replication
 			TPromise<FJoinReplicatedSessionResult> JoinSessionPromise,
 			TSharedRef<IConcertClientSession> LiveSession,
 			IConcertClientReplicationBridge& ReplicationBridge UE_LIFETIMEBOUND,
-			FReplicationManager& Owner UE_LIFETIMEBOUND
+			FReplicationManager& Owner UE_LIFETIMEBOUND,
+			EConcertSyncSessionFlags SessionFlags
 			);
 		virtual ~FReplicationManagerState_Handshaking() override;
 		
@@ -45,6 +44,8 @@ namespace UE::ConcertSyncClient::Replication
 		TSharedRef<IConcertClientSession> LiveSession;
 		/** Passed to FReplicationManagerState_Handshaking */
 		IConcertClientReplicationBridge& ReplicationBridge;
+		/** Passed to FReplicationManagerState_Handshaking */
+		const EConcertSyncSessionFlags SessionFlags;
 
 		//~ Begin FReplicationManagerState Interface
 		virtual void OnEnterState() override;
diff --git a/Engine/Plugins/Developer/Concert/ConcertSync/ConcertSyncClient/Source/ConcertSyncClient/Private/Replication/Misc/LocalSyncControl.h b/Engine/Plugins/Developer/Concert/ConcertSync/ConcertSyncClient/Source/ConcertSyncClient/Private/Replication/Misc/LocalSyncControl.h
--- a/Engine/Plugins/Developer/Concert/ConcertSync/ConcertSyncClient/Source/ConcertSyncClient/Private/Replication/Misc/LocalSyncControl.h
+++ b/Engine/Plugins/Developer/Concert/ConcertSync/ConcertSyncClient/Source/ConcertSyncClient/Private/Replication/Misc/LocalSyncControl.h
@@ -38,6 +38,7 @@ namespace UE::ConcertSyncClient::Replication
 		using FSyncControlState::Num;
 		using FSyncControlState::EnumerateAllowedObjects;
 		using FSyncControlState::EnumerateChanges;
+		using FPredictedObjectRemoval = FSyncControlState::FPredictedObjectRemoval;
 		// No using for FSyncControlState::Aggregate because we want to limit usage to below ProcessXChange functions
 
 		/** Stores the given change in this data structure. */
@@ -66,6 +67,47 @@ namespace UE::ConcertSyncClient::Replication
 			AppendChanges(Request);
 			OnPostSyncControlChangedDelegate.Broadcast();
 		}
+
+		/**
+		 * Applies the implicit changes made by the request assuming the request will be accepted.
+		 * @return The removed objects to be passed to ApplyOrRevertMuteResponse.
+		 */
+		FPredictedObjectRemoval PredictAndApplyMuteChanges(const FConcertReplication_ChangeMuteState_Request& Request)
+		{
+			bool bMadeChange = false;
+			FPredictedObjectRemoval Predicition = AppendChanges(Request, [this, &bMadeChange](const FConcertObjectInStreamID&)
+			{
+				if (!bMadeChange)
+				{
+					OnPreSyncControlChangedDelegate.Broadcast();
+				}
+				bMadeChange = true;
+			});
+			if (bMadeChange)
+			{
+				OnPostSyncControlChangedDelegate.Broadcast();
+			}
+
+			return Predicition;
+		}
+
+		/** Either reverts previous changes made if the request was rejected, or applies the sync control returned by the server otherwise. */
+		void ApplyOrRevertMuteResponse(const FPredictedObjectRemoval& RemovedByRequest, const FConcertReplication_ChangeMuteState_Response& Response)
+		{
+			bool bMadeChange = false;
+			AppendChanges(RemovedByRequest, Response, [this, &bMadeChange](const FConcertObjectInStreamID&)
+			{
+				if (!bMadeChange)
+				{
+					OnPreSyncControlChangedDelegate.Broadcast();
+				}
+				bMadeChange = true;
+			});
+			if (bMadeChange)
+			{
+				OnPostSyncControlChangedDelegate.Broadcast();
+			}
+		}
 		
 		DECLARE_MULTICAST_DELEGATE(FSyncControlChanged);
 		FSyncControlChanged& OnPreSyncControlChanged() { return OnPreSyncControlChangedDelegate; }
diff --git a/Engine/Plugins/Developer/Concert/ConcertSync/ConcertSyncClient/Source/ConcertSyncClient/Private/Replication/ReplicationTestInterface.cpp b/Engine/Plugins/Developer/Concert/ConcertSync/ConcertSyncClient/Source/ConcertSyncClient/Private/Replication/ReplicationTestInterface.cpp
--- a/Engine/Plugins/Developer/Concert/ConcertSync/ConcertSyncClient/Source/ConcertSyncClient/Private/Replication/ReplicationTestInterface.cpp
+++ b/Engine/Plugins/Developer/Concert/ConcertSync/ConcertSyncClient/Source/ConcertSyncClient/Private/Replication/ReplicationTestInterface.cpp
@@ -1,6 +1,7 @@
 // Copyright Epic Games, Inc. All Rights Reserved.
 
 #include "ConcertClientReplicationBridge.h"
+#include "ConcertSyncSessionFlags.h"
 #include "Replication/IConcertClientReplicationBridge.h"
 #include "Replication/IConcertClientReplicationManager.h"
 #include "Manager/ReplicationManager.h"
@@ -15,10 +16,11 @@ namespace UE::ConcertSyncClient::TestInterface
 {
 	CONCERTSYNCCLIENT_API TSharedRef<IConcertClientReplicationManager> CreateClientReplicationManager(
 		TSharedRef<IConcertClientSession> InLiveSession,
-		IConcertClientReplicationBridge& InBridge
+		IConcertClientReplicationBridge& InBridge,
+		EConcertSyncSessionFlags SessionFlags
 		)
 	{
-		const TSharedRef<Replication::FReplicationManager> Result = MakeShared<Replication::FReplicationManager>(MoveTemp(InLiveSession), InBridge);
+		const TSharedRef<Replication::FReplicationManager> Result = MakeShared<Replication::FReplicationManager>(MoveTemp(InLiveSession), InBridge, SessionFlags);
 		Result->StartAcceptingJoinRequests();
 		return Result;
 	}
diff --git a/Engine/Plugins/Developer/Concert/ConcertSync/ConcertSyncClient/Source/ConcertSyncClient/Public/Replication/IConcertClientReplicationManager.h b/Engine/Plugins/Developer/Concert/ConcertSync/ConcertSyncClient/Source/ConcertSyncClient/Public/Replication/IConcertClientReplicationManager.h
--- a/Engine/Plugins/Developer/Concert/ConcertSync/ConcertSyncClient/Source/ConcertSyncClient/Public/Replication/IConcertClientReplicationManager.h
+++ b/Engine/Plugins/Developer/Concert/ConcertSync/ConcertSyncClient/Source/ConcertSyncClient/Public/Replication/IConcertClientReplicationManager.h
@@ -3,10 +3,11 @@
 #pragma once
 
 #include "Misc/EBreakBehavior.h"
-#include "Replication/Messages/Handshake.h"
 #include "Replication/Messages/ChangeAuthority.h"
 #include "Replication/Messages/ChangeStream.h"
 #include "Replication/Messages/ClientQuery.h"
+#include "Replication/Messages/Handshake.h"
+#include "Replication/Messages/Muting.h"
 
 template<typename ResultType>
 class TFuture;
@@ -82,9 +83,9 @@ public:
 	 */
 	virtual TFuture<FConcertReplication_ChangeAuthority_Response> RequestAuthorityChange(FConcertReplication_ChangeAuthority_Request Args) = 0;
 	/** Util function that will request authority for all streams for the given objects. */
-	TFuture<FConcertReplication_ChangeAuthority_Response> TakeAuthorityOver(TArrayView<const FSoftObjectPath> Objects);
+	TFuture<FConcertReplication_ChangeAuthority_Response> TakeAuthorityOver(TConstArrayView<FSoftObjectPath> Objects);
 	/** Util function that will let go over all authority of the given objects. */
-	TFuture<FConcertReplication_ChangeAuthority_Response> ReleaseAuthorityOf(TArrayView<const FSoftObjectPath> Objects);
+	TFuture<FConcertReplication_ChangeAuthority_Response> ReleaseAuthorityOf(TConstArrayView<FSoftObjectPath> Objects);
 
 	enum class EAuthorityEnumerationResult { NoAuthorityAvailable, Iterated };
 	/**
@@ -119,6 +120,19 @@ public:
 	 */
 	virtual TFuture<FConcertReplication_ChangeStream_Response> ChangeStream(FConcertReplication_ChangeStream_Request Args) = 0;
 
+	/**
+	 * Requests to change the global mute state of objects.
+	 * Can only be done if the session has the EConcertSyncSessionFlags::ShouldAllowGlobalMuting flag.
+	 */
+	virtual TFuture<FConcertReplication_ChangeMuteState_Response> ChangeMuteState(FConcertReplication_ChangeMuteState_Request Request) = 0;
+	/** Util function that will mute all of Objects. */
+	TFuture<FConcertReplication_ChangeMuteState_Response> MuteObjects(TConstArrayView<FSoftObjectPath> Objects, EConcertReplicationMuteOption Flags = EConcertReplicationMuteOption::ObjectAndSubobjects);
+	/** Util function that will unmute all of Objects. */
+	TFuture<FConcertReplication_ChangeMuteState_Response> UnmuteObjects(TSet<FSoftObjectPath> Objects, EConcertReplicationMuteOption Flags = EConcertReplicationMuteOption::ObjectAndSubobjects);
+	/** Gets the global mute state. */
+	virtual TFuture<FConcertReplication_QueryMuteState_Response> QueryMuteState(FConcertReplication_QueryMuteState_Request Request = {}) = 0;
+	TFuture<FConcertReplication_QueryMuteState_Response> QueryMuteState(TSet<FSoftObjectPath> Objects);
+
 	DECLARE_MULTICAST_DELEGATE_TwoParams(FOnPreStreamsChanged,
 		const FConcertReplication_ChangeStream_Request&,
 		const FConcertReplication_ChangeStream_Response&
diff --git a/Engine/Plugins/Developer/Concert/ConcertSync/ConcertSyncCore/Source/ConcertSyncCore/Private/Replication/ChangeStreamSharedUtils.cpp b/Engine/Plugins/Developer/Concert/ConcertSync/ConcertSyncCore/Source/ConcertSyncCore/Private/Replication/ChangeStreamSharedUtils.cpp
--- a/Engine/Plugins/Developer/Concert/ConcertSync/ConcertSyncCore/Source/ConcertSyncCore/Private/Replication/ChangeStreamSharedUtils.cpp
+++ b/Engine/Plugins/Developer/Concert/ConcertSync/ConcertSyncCore/Source/ConcertSyncCore/Private/Replication/ChangeStreamSharedUtils.cpp
@@ -5,11 +5,12 @@
 #include "Misc/EBreakBehavior.h"
 #include "Replication/Data/ObjectIds.h"
 #include "Replication/Messages/ChangeStream.h"
-#include "Replication/Messages/ObjectReplication.h"
+
+#include "Algo/NoneOf.h"
 
 namespace UE::ConcertSyncCore::Replication::ChangeStreamUtils
 {
-	void ForEachObjectLosingAuthority(
+	void ForEachRemovedObject(
 		const FConcertReplication_ChangeStream_Request& Request,
 		const TArray<FConcertReplicationStream>& ExistingStreams,
 		TFunctionRef<EBreakBehavior(const FConcertObjectInStreamID&)> Callback
@@ -41,6 +42,36 @@ namespace UE::ConcertSyncCore::Replication::ChangeStreamUtils
 			}
 		}
 	}
+
+	void ForEachAddedObject(
+		const FConcertReplication_ChangeStream_Request& Request,
+		const TArray<FConcertReplicationStream>& ExistingStreams,
+		TFunctionRef<EBreakBehavior(const FConcertObjectInStreamID&)> Callback
+		)
+	{
+		for (const TPair<FConcertObjectInStreamID, FConcertReplication_ChangeStream_PutObject>& PutObject : Request.ObjectsToPut)
+		{
+			const bool bIsNewObject = Algo::NoneOf(ExistingStreams, [&PutObject](const FConcertReplicationStream& Stream)
+			{
+				return Stream.BaseDescription.ReplicationMap.HasProperties(PutObject.Key.Object);
+			});
+			if (bIsNewObject && Callback(PutObject.Key) == EBreakBehavior::Break)
+			{
+				return;
+			}
+		}
+		
+		for (const FConcertReplicationStream& AddedStream : Request.StreamsToAdd)
+		{
+			for (const TPair<FSoftObjectPath, FConcertReplicatedObjectInfo>& ObjectInfo : AddedStream.BaseDescription.ReplicationMap.ReplicatedObjects)
+			{
+				if (Callback({ AddedStream.BaseDescription.Identifier, ObjectInfo.Key }) == EBreakBehavior::Break)
+				{
+					return;
+				}
+			}
+		}
+	}
 	
 	void ApplyValidatedRequest(const FConcertReplication_ChangeStream_Request& Request, TArray<FConcertReplicationStream>& StreamsToModify)
 	{
diff --git a/Engine/Plugins/Developer/Concert/ConcertSync/ConcertSyncCore/Source/ConcertSyncCore/Private/Replication/Misc/ReplicatedObjectHierarchyCache.cpp b/Engine/Plugins/Developer/Concert/ConcertSync/ConcertSyncCore/Source/ConcertSyncCore/Private/Replication/Misc/ReplicatedObjectHierarchyCache.cpp
--- a/Engine/Plugins/Developer/Concert/ConcertSync/ConcertSyncCore/Source/ConcertSyncCore/Private/Replication/Misc/ReplicatedObjectHierarchyCache.cpp
+++ b/Engine/Plugins/Developer/Concert/ConcertSync/ConcertSyncCore/Source/ConcertSyncCore/Private/Replication/Misc/ReplicatedObjectHierarchyCache.cpp
@@ -37,7 +37,7 @@ namespace UE::ConcertSyncCore
 		}
 	}
 
-	void FReplicatedObjectHierarchyCache::OnPreClientLeft(const FGuid& ClientId, TConstArrayView<FConcertReplicationStream> Streams)
+	void FReplicatedObjectHierarchyCache::OnPostClientLeft(const FGuid& ClientId, TConstArrayView<FConcertReplicationStream> Streams)
 	{
 		for (const FConcertReplicationStream& Stream : Streams)
 		{
diff --git a/Engine/Plugins/Developer/Concert/ConcertSync/ConcertSyncCore/Source/ConcertSyncCore/Public/ConcertSyncSessionFlags.h b/Engine/Plugins/Developer/Concert/ConcertSync/ConcertSyncCore/Source/ConcertSyncCore/Public/ConcertSyncSessionFlags.h
--- a/Engine/Plugins/Developer/Concert/ConcertSync/ConcertSyncCore/Source/ConcertSyncCore/Public/ConcertSyncSessionFlags.h
+++ b/Engine/Plugins/Developer/Concert/ConcertSync/ConcertSyncCore/Source/ConcertSyncCore/Public/ConcertSyncSessionFlags.h
@@ -46,13 +46,18 @@ enum class EConcertSyncSessionFlags : uint32
 	ShouldSendPackageAutoSaves = 1<<12,
 	/** Use a sandbox to hold package writes during a session (requires EnablePackages) */
 	ShouldUsePackageSandbox = 1<<13,
+	/**
+	 * For Concert Replication, clients are allowed to globally mute objects.
+	 * More formally, allows the use of FConcertReplication_ChangeMuteState_Requests.
+	 */
+	ShouldAllowGlobalMuting = 1<<14,
 
 	/**
 	 * Defaults for different scenarios
 	 */
 	
 	/** Default settings for a multi-user session */
-	Default_MultiUserSession = EnableLiveSync | EnableConnectionHistory | EnableLocking | EnableTransactions | EnablePackages | EnableSequencer | EnablePresence | ShouldSendTransactionSnapshots | ShouldUsePackageSandbox | EnableReplication,
+	Default_MultiUserSession = EnableLiveSync | EnableConnectionHistory | EnableLocking | EnableTransactions | EnablePackages | EnableSequencer | EnablePresence | ShouldSendTransactionSnapshots | ShouldUsePackageSandbox | EnableReplication | ShouldAllowGlobalMuting,
 	/** Default settings for a disaster recovery session */
 	Default_DisasterRecoverySession = EnableTransactions | EnablePackages | EnableFileSharing | ShouldDiscardTransactionsOnPackageUnload | ShouldSendPackagePristineState | ShouldSendPackageAutoSaves,
 };
diff --git a/Engine/Plugins/Developer/Concert/ConcertSync/ConcertSyncCore/Source/ConcertSyncCore/Public/Replication/ChangeStreamSharedUtils.h b/Engine/Plugins/Developer/Concert/ConcertSync/ConcertSyncCore/Source/ConcertSyncCore/Public/Replication/ChangeStreamSharedUtils.h
--- a/Engine/Plugins/Developer/Concert/ConcertSync/ConcertSyncCore/Source/ConcertSyncCore/Public/Replication/ChangeStreamSharedUtils.h
+++ b/Engine/Plugins/Developer/Concert/ConcertSync/ConcertSyncCore/Source/ConcertSyncCore/Public/Replication/ChangeStreamSharedUtils.h
@@ -20,12 +20,27 @@ namespace UE::ConcertSyncCore::Replication::ChangeStreamUtils
 {
 	/**
 	 * Iterates through all objects that are removed and therefore the requesting client would lose authority over.
+	 * This function assumes that Request was validated.
 	 * 
 	 * @param Request The request to parse
 	 * @param ExistingStreams The streams the request would be applied to
 	 * @param Callback The callback to process the objects
 	 */
-	CONCERTSYNCCORE_API void ForEachObjectLosingAuthority(
+	CONCERTSYNCCORE_API void ForEachRemovedObject(
+		const FConcertReplication_ChangeStream_Request& Request,
+		const TArray<FConcertReplicationStream>& ExistingStreams,
+		TFunctionRef<EBreakBehavior(const FConcertObjectInStreamID&)> Callback
+		);
+
+	/**
+	 * Iterates through all objects that are added by Request but not yet in ExistingStreams.
+	 * This function assumes that Request was validated.
+	 *
+	 * @param Request The request to parse
+	 * @param ExistingStreams The streams the request would be applied to
+	 * @param Callback The callback to process the objects
+	 */
+	CONCERTSYNCCORE_API void ForEachAddedObject(
 		const FConcertReplication_ChangeStream_Request& Request,
 		const TArray<FConcertReplicationStream>& ExistingStreams,
 		TFunctionRef<EBreakBehavior(const FConcertObjectInStreamID&)> Callback
diff --git a/Engine/Plugins/Developer/Concert/ConcertSync/ConcertSyncCore/Source/ConcertSyncCore/Public/Replication/Messages/ChangeAuthority.h b/Engine/Plugins/Developer/Concert/ConcertSync/ConcertSyncCore/Source/ConcertSyncCore/Public/Replication/Messages/ChangeAuthority.h
--- a/Engine/Plugins/Developer/Concert/ConcertSync/ConcertSyncCore/Source/ConcertSyncCore/Public/Replication/Messages/ChangeAuthority.h
+++ b/Engine/Plugins/Developer/Concert/ConcertSync/ConcertSyncCore/Source/ConcertSyncCore/Public/Replication/Messages/ChangeAuthority.h
@@ -67,7 +67,7 @@ struct FConcertReplication_ChangeAuthority_Response
 
 	/**
 	 * Contains the new sync control for objects you requested authority over; objects you released authority over are not contained because they can
-	 * be implicitly determined using FSyncControlState::Aggregate.
+	 * be implicitly determined using FSyncControlState::AppendChanges.
 	 * 
 	 * @see FSyncControlState::Aggregate.
 	 * @see FConcertReplication_ChangeSyncControl
diff --git a/Engine/Plugins/Developer/Concert/ConcertSync/ConcertSyncCore/Source/ConcertSyncCore/Public/Replication/Messages/Muting.h b/Engine/Plugins/Developer/Concert/ConcertSync/ConcertSyncCore/Source/ConcertSyncCore/Public/Replication/Messages/Muting.h
new file mode 100644
--- /dev/null
+++ b/Engine/Plugins/Developer/Concert/ConcertSync/ConcertSyncCore/Source/ConcertSyncCore/Public/Replication/Messages/Muting.h
@@ -0,0 +1,180 @@
+﻿// Copyright Epic Games, Inc. All Rights Reserved.
+
+#pragma once
+
+#include "EReplicationResponseErrorCode.h"
+#include "SyncControl.h"
+
+#include "HAL/Platform.h"
+#include "Misc/EnumClassFlags.h"
+#include "Muting.generated.h"
+
+UENUM()
+enum class EConcertReplicationMuteOption : uint8
+{
+	/** Only the specified object is affected; subobjects are not. */
+	OnlyObject,
+	/** The setting applies for all subobjects of the object as well. If subobjects are added in the future, the setting will apply to them as well. */
+	ObjectAndSubobjects = 1 << 0,
+};
+
+namespace UE::ConcertSyncCore
+{
+	/** @return Whether subobjects are affected by this Option. Used for future proofing (in case new entries are added to the enum). */
+	inline bool AffectSubobjects(EConcertReplicationMuteOption Option) { return Option == EConcertReplicationMuteOption::ObjectAndSubobjects; }
+}
+
+/** Describes how an object is to be muted. */
+USTRUCT()
+struct FConcertReplication_ObjectMuteSetting
+{
+	GENERATED_BODY()
+
+	/** Modify the default behavior of the mute operation */
+	UPROPERTY()
+	EConcertReplicationMuteOption Flags = EConcertReplicationMuteOption::ObjectAndSubobjects;
+
+	friend bool operator==(const FConcertReplication_ObjectMuteSetting& Left, const FConcertReplication_ObjectMuteSetting& Right) = default;
+};
+
+/**
+ * A request to globally pause / resume replication of objects.
+ * Muted objects will not be replicated to any clients.
+ * For security reasons, muting is only available in sessions with the EConcertSyncSessionFlags::ShouldAllowGlobalMuting flag set.
+ *
+ * If a muted objects is unregistered (i.e. by a stream change request), the mute state is also removed.
+ * That the object is later re-registered, it will not be muted.
+ *
+ * Requests are processed atomically: it is completed either completely or not at all.
+ * The following operations are invalid:
+ * - Unmuting / muting an object without the IncludeSubobjects flag if that object is not registered in any stream.
+ * - Unmuting / muting an object with the IncludeSubobjects flag if that object is not registered in any stream, and that object has no subobjects registered in any streams.
+ * The reason for making these operations invalid is to prevent memory leaks on the server (e.g. client unregisters an object and forgets to remove the mute state, too).
+ * The following operations are valid:
+ * - Unmuting an object that is already unmuted (no-op)
+ * - Muting an object that is already muted (no-op); changing the EConcertReplicationMuteFlags is not a no-op.
+ *
+ * Sync control:
+ * This request may update sync control. The following rules apply:
+ * - This request may cause FConcertReplication_ChangeSyncControl events to be sent to OTHER clients.
+ * - This request will NOT cause any FConcertReplication_ChangeSyncControl event to be sent to the requesting client:
+ *		- When muting, the requesting client is expected to infer that sync control is lost for the muted objects.
+ *		- When unmuting, the server fills FConcertReplication_ChangeMuteState_Response::SyncControl with the object the requesting client has received sync control over.
+ */
+USTRUCT()
+struct FConcertReplication_ChangeMuteState_Request
+{
+	GENERATED_BODY()
+
+	/** The objects to explicitly mute. */
+	UPROPERTY()
+	TMap<FSoftObjectPath, FConcertReplication_ObjectMuteSetting> ObjectsToMute;
+
+	/**
+	 * The objects to explicitly unmute.
+	 * 
+	 * If a specified object is implicitly muted, i.e. one of its outers is muted with the EConcertReplicationMuteFlags::IncludeSubobjects flag,
+	 * this will allow it to replicate again. If you specify the EConcertReplicationMuteFlags::IncludeSubobjects flag, then its subobjects will
+	 * also be allowed to be replicated again.
+	 *
+	 * If a specified object is explicitly muted, then it and all the subobjects that were implicitly muted because of it are unmuted;
+	 * in this case, it does not matter whether EConcertReplicationMuteFlags::IncludeSubobjects flag is set.
+	 */
+	UPROPERTY()
+	TMap<FSoftObjectPath, FConcertReplication_ObjectMuteSetting> ObjectsToUnmute;
+
+	/** @return Whether this requests makes no changes */
+	bool IsEmpty() const { return ObjectsToMute.IsEmpty() && ObjectsToUnmute.IsEmpty(); }
+};
+
+/** Result of a FConcertReplication_ChangeMuteState_Request. */
+UENUM()
+enum class EConcertReplicationMuteErrorCode : uint8
+{
+	/**
+	 * Value to set when default constructed (Concert does this to responses when they timeout or the request in invalid).
+	 * Happens when timed out or because EConcertSyncSessionFlags::ShouldAllowGlobalMuting is not set.
+	 */
+	Timeout,
+
+	/** Changes were applied. */
+	Accepted,
+	
+	/** No change were made. The request was malformed. */
+	Rejected
+};
+
+/** If the ErrorCode != Accepted, then no changes were made on the server. */
+USTRUCT()
+struct FConcertReplication_ChangeMuteState_Response
+{
+	GENERATED_BODY()
+
+	UPROPERTY()
+	EConcertReplicationMuteErrorCode ErrorCode = EConcertReplicationMuteErrorCode::Timeout;
+
+	/** The objects that caused the request to be rejected. Only valid if EConcertReplicationMuteErrorCode::Rejected. */
+	UPROPERTY()
+	TSet<FSoftObjectPath> RejectionReasons;
+
+	/**
+	 * If the request unmuted objects, this contains the objects that the requester gained sync control over.
+	 * 
+	 * @note When muting, the requesting client is expected to infer that sync control is lost for the muted objects.
+	 * Hence, this will not contain objects the requester lost sync control over.
+	 */
+	UPROPERTY()
+	FConcertReplication_ChangeSyncControl SyncControl;
+
+	bool IsSuccess() const { return ErrorCode == EConcertReplicationMuteErrorCode::Accepted && ensure(RejectionReasons.IsEmpty()); }
+	bool IsFailure() const { return !IsSuccess(); }
+};
+
+/** Queries the effective mute state for.*/
+USTRUCT()
+struct FConcertReplication_QueryMuteState_Request
+{
+	GENERATED_BODY()
+
+	/** Specifies the objects for which to get the mute states. If left empty, get the states of all objects. */
+	UPROPERTY()
+	TSet<FSoftObjectPath> QueriedObjects;
+
+	bool WantsAllObjects() const { return QueriedObjects.IsEmpty(); }
+};
+
+USTRUCT()
+struct FConcertReplication_QueryMuteState_Response
+{
+	GENERATED_BODY()
+
+	UPROPERTY()
+	EReplicationResponseErrorCode ErrorCode = EReplicationResponseErrorCode::Timeout;
+
+	/**
+	 * Objects that were explicitly muted, i.e. specified with FConcertReplication_ChangeMuteState_Request::ObjectsToMute.
+	 * */
+	UPROPERTY()
+	TMap<FSoftObjectPath, FConcertReplication_ObjectMuteSetting> ExplicitlyMutedObjects;
+
+	/**
+	 * After an outer object is muted with EConcertReplicationMuteFlags::IncludeSubobjects, this contains the objects that were explicitly unmuted,
+	 * i.e. specified with FConcertReplication_ChangeMuteState_Request::ObjectsToUnmute.
+	 */
+	UPROPERTY()
+	TMap<FSoftObjectPath, FConcertReplication_ObjectMuteSetting> ExplicitlyUnmutedObjects;
+
+	/**
+	 * Objects that were implicitly muted, i.e. some their outers are explicitly muted with EConcertReplicationMuteFlags::IncludeSubobjects.
+	 */
+	UPROPERTY()
+	TSet<FSoftObjectPath> ImplicitlyMutedObjects;
+
+	/** Objects that were implicitly unmuted, i.e. 1. a parent is muted and 2. a child of parent is unmuted with EConcertReplicationMuteFlags::IncludeSubobjects */
+	UPROPERTY()
+	TSet<FSoftObjectPath> ImplicitlyUnmutedObjects;
+
+	bool IsSuccess() const { return ErrorCode == EReplicationResponseErrorCode::Handled; }
+	bool IsFailure() const { return !IsSuccess(); }
+	bool IsEmpty() const { return ExplicitlyMutedObjects.IsEmpty() && ExplicitlyUnmutedObjects.IsEmpty() && ImplicitlyMutedObjects.IsEmpty(); }
+};
\ No newline at end of file
diff --git a/Engine/Plugins/Developer/Concert/ConcertSync/ConcertSyncCore/Source/ConcertSyncCore/Public/Replication/Messages/SyncControl.h b/Engine/Plugins/Developer/Concert/ConcertSync/ConcertSyncCore/Source/ConcertSyncCore/Public/Replication/Messages/SyncControl.h
--- a/Engine/Plugins/Developer/Concert/ConcertSync/ConcertSyncCore/Source/ConcertSyncCore/Public/Replication/Messages/SyncControl.h
+++ b/Engine/Plugins/Developer/Concert/ConcertSync/ConcertSyncCore/Source/ConcertSyncCore/Public/Replication/Messages/SyncControl.h
@@ -2,6 +2,7 @@
 
 #pragma once
 
+#include "Algo/AnyOf.h"
 #include "Replication/Data/ObjectIds.h"
 #include "SyncControl.generated.h"
 
@@ -41,4 +42,7 @@ struct FConcertReplication_ChangeSyncControl
 	TMap<FConcertObjectInStreamID, bool> NewControlStates;
 
 	bool IsEmpty() const { return NewControlStates.IsEmpty(); }
+	
+	bool DoesAtLeastOneObjectGainSyncControl() const { return Algo::AnyOf(NewControlStates, [](const TPair<FConcertObjectInStreamID, bool>& Pair){ return Pair.Value; }); }
+	bool DoesAtLeastOneObjectLoseSyncControl() const { return Algo::AnyOf(NewControlStates, [](const TPair<FConcertObjectInStreamID, bool>& Pair){ return !Pair.Value; }); }
 };
\ No newline at end of file
diff --git a/Engine/Plugins/Developer/Concert/ConcertSync/ConcertSyncCore/Source/ConcertSyncCore/Public/Replication/Misc/ReplicatedObjectHierarchyCache.h b/Engine/Plugins/Developer/Concert/ConcertSync/ConcertSyncCore/Source/ConcertSyncCore/Public/Replication/Misc/ReplicatedObjectHierarchyCache.h
--- a/Engine/Plugins/Developer/Concert/ConcertSync/ConcertSyncCore/Source/ConcertSyncCore/Public/Replication/Misc/ReplicatedObjectHierarchyCache.h
+++ b/Engine/Plugins/Developer/Concert/ConcertSync/ConcertSyncCore/Source/ConcertSyncCore/Public/Replication/Misc/ReplicatedObjectHierarchyCache.h
@@ -46,7 +46,7 @@ namespace UE::ConcertSyncCore
 
 		// Validated network events that modify the known hierarchy of objects.
 		void OnJoin(const FGuid& ClientId, const FConcertReplication_Join_Request& Request);
-		void OnPreClientLeft(const FGuid& ClientId, TConstArrayView<FConcertReplicationStream> Streams);
+		void OnPostClientLeft(const FGuid& ClientId, TConstArrayView<FConcertReplicationStream> Streams);
 		void OnChangeStreams(const FGuid& ClientId, TConstArrayView<FConcertObjectInStreamID> AddedObjects, TConstArrayView<FConcertObjectInStreamID> RemovedObjects);
 
 	private:
diff --git a/Engine/Plugins/Developer/Concert/ConcertSync/ConcertSyncCore/Source/ConcertSyncCore/Public/Replication/SyncControlState.h b/Engine/Plugins/Developer/Concert/ConcertSync/ConcertSyncCore/Source/ConcertSyncCore/Public/Replication/SyncControlState.h
--- a/Engine/Plugins/Developer/Concert/ConcertSync/ConcertSyncCore/Source/ConcertSyncCore/Public/Replication/SyncControlState.h
+++ b/Engine/Plugins/Developer/Concert/ConcertSync/ConcertSyncCore/Source/ConcertSyncCore/Public/Replication/SyncControlState.h
@@ -13,6 +13,8 @@
 
 #include <type_traits>
 
+#include "Messages/Muting.h"
+
 namespace UE::ConcertSyncCore::Replication
 {
 	template<typename T>
@@ -28,6 +30,12 @@ namespace UE::ConcertSyncCore::Replication
 	{
 	public:
 
+		class FPredictedObjectRemoval
+		{
+			friend class FSyncControlState;
+			TSet<FConcertObjectInStreamID> Objects;
+		};
+
 		FSyncControlState() = default;
 		
 		// Not explicit to make it useful to use
@@ -65,6 +73,24 @@ namespace UE::ConcertSyncCore::Replication
 		template<CObjectInStreamInvocable TOnDisallowed>
 		void AppendChanges(const FConcertReplication_ChangeStream_Request& Request, TOnDisallowed&& OnDisallowed);
 		void AppendChanges(const FConcertReplication_ChangeStream_Request& Request);
+
+		/**
+		 * Combines implicit sync control changes caused by the local client sending a mute event.
+		 * @return The objects that were predictively removed by the request. Pass this back into AppendChanges once you get the response.
+		 */
+		template<CObjectInStreamInvocable TOnDisallowed>
+		FPredictedObjectRemoval AppendChanges(const FConcertReplication_ChangeMuteState_Request& Request, TOnDisallowed&& OnDisallowed);
+		FPredictedObjectRemoval AppendChanges(const FConcertReplication_ChangeMuteState_Request& Request);
+
+		/**
+		 * Looks at the response:
+		 * - if the change failed, reverts the predictively removed sync control
+		 * - if the change succeeded, appends the contained sync control
+		 */
+		template<CObjectInStreamInvocable TOnAllowed>
+		void AppendChanges(const FPredictedObjectRemoval& ObjectsRemovedInRequest, const FConcertReplication_ChangeMuteState_Response& Response, TOnAllowed&& OnAllowed);
+		void AppendChanges(const FPredictedObjectRemoval& ObjectsRemovedInRequest, const FConcertReplication_ChangeMuteState_Response& Response);
+		
 		
 		friend bool operator==(const FSyncControlState& Left, const FSyncControlState& Right)
 		{
@@ -185,5 +211,77 @@ namespace UE::ConcertSyncCore::Replication
 	{
 		AppendChanges(Request, [](const FConcertObjectInStreamID&){});
 	}
+
+	template <CObjectInStreamInvocable TOnDisallowed>
+	FSyncControlState::FPredictedObjectRemoval FSyncControlState::AppendChanges(const FConcertReplication_ChangeMuteState_Request& Request, TOnDisallowed&& OnDisallowed)
+	{
+		FPredictedObjectRemoval RemovedObjects;
+		const auto Remove = [&OnDisallowed, &RemovedObjects](TSet<FConcertObjectInStreamID>::TIterator& It)
+		{
+			const FConcertObjectInStreamID Object = *It;
+			It.RemoveCurrent();
+						
+			RemovedObjects.Objects.Add(Object);
+			OnDisallowed(Object);
+		};
+		
+		for (const TPair<FSoftObjectPath, FConcertReplication_ObjectMuteSetting>& ObjectInfo : Request.ObjectsToMute)
+		{
+			const FSoftObjectPath& MutedObject = ObjectInfo.Key;
+			const EConcertReplicationMuteOption MuteOption = ObjectInfo.Value.Flags;
+			
+			if (AffectSubobjects(MuteOption))
+			{
+				const FString MutedAsString = MutedObject.ToString();
+				for (auto It = AllowedObjects.CreateIterator(); It; ++It)
+				{
+					const bool bRemove = It->Object.ToString().Contains(MutedObject.ToString());
+					if (bRemove)
+					{
+						Remove(It);
+					}
+				}
+			}
+			else
+			{
+				for (auto It = AllowedObjects.CreateIterator(); It; ++It)
+				{
+					if (It->Object == MutedObject)
+					{
+						Remove(It);
+					}
+				}
+			}
+		}
+		return RemovedObjects;
+	}
+
+	inline FSyncControlState::FPredictedObjectRemoval FSyncControlState::AppendChanges(const FConcertReplication_ChangeMuteState_Request& Request)
+	{
+		return AppendChanges(Request, [](const FConcertObjectInStreamID&){});
+	}
+
+	template <CObjectInStreamInvocable TOnAllowed>
+	void FSyncControlState::AppendChanges(const FPredictedObjectRemoval& ObjectsRemovedInRequest, const FConcertReplication_ChangeMuteState_Response& Response, TOnAllowed&& OnAllowed)
+	{
+		if (Response.IsSuccess())
+		{
+			AppendChanges(Response.SyncControl,
+				OnAllowed,
+				[](const FConcertObjectInStreamID&)
+				{
+					ensureAlwaysMsgf(false, TEXT("By contract, objects losing sync control are not supposed to be listed here. FConcertReplication_ChangeMuteState_Response::SyncControl docoumentation."));
+				});
+		}
+		else
+		{
+			AllowedObjects.Append(ObjectsRemovedInRequest.Objects);
+		}
+	}
+
+	inline void FSyncControlState::AppendChanges(const FPredictedObjectRemoval& ObjectsRemovedInRequest, const FConcertReplication_ChangeMuteState_Response& Response)
+	{
+		AppendChanges(ObjectsRemovedInRequest, Response, [](const FConcertObjectInStreamID&){});
+	}
 }
 
diff --git a/Engine/Plugins/Developer/Concert/ConcertSync/ConcertSyncServer/Source/ConcertSyncServer/Private/ConcertSyncServer.cpp b/Engine/Plugins/Developer/Concert/ConcertSync/ConcertSyncServer/Source/ConcertSyncServer/Private/ConcertSyncServer.cpp
--- a/Engine/Plugins/Developer/Concert/ConcertSync/ConcertSyncServer/Source/ConcertSyncServer/Private/ConcertSyncServer.cpp
+++ b/Engine/Plugins/Developer/Concert/ConcertSync/ConcertSyncServer/Source/ConcertSyncServer/Private/ConcertSyncServer.cpp
@@ -526,15 +526,19 @@ void FConcertSyncServer::DestroySequencerManager(const TSharedRef<FConcertSyncSe
 	LiveSessionSequencerManagers.Remove(InLiveSession->GetSession().GetId());
 }
 
-void FConcertSyncServer::CreateReplicationManager(TSharedRef<IConcertServerSession> InLiveSession)
+void FConcertSyncServer::CreateReplicationManager(const TSharedRef<IConcertServerSession>& InSession, EConcertSyncSessionFlags InSessionFlags)
 {
-	DestroyReplicationManager(InLiveSession);
-	LiveSessionReplicationManagers.Add(InLiveSession->GetId(), MakeShared<UE::ConcertSyncServer::Replication::FConcertServerReplicationManager>(MoveTemp(InLiveSession)));
+	const FGuid& SessionId = InSession->GetId(); 
+	DestroyReplicationManager(SessionId);
+	LiveSessionReplicationManagers.Add(
+		SessionId,
+		MakeShared<UE::ConcertSyncServer::Replication::FConcertServerReplicationManager>(InSession, InSessionFlags)
+		);
 }
 
-void FConcertSyncServer::DestroyReplicationManager(const TSharedRef<IConcertServerSession>& InLiveSession)
+void FConcertSyncServer::DestroyReplicationManager(const FGuid& SessionId)
 {
-	LiveSessionReplicationManagers.Remove(InLiveSession->GetId());
+	LiveSessionReplicationManagers.Remove(SessionId);
 }
 
 bool FConcertSyncServer::CreateLiveSession(const TSharedRef<IConcertServerSession>& InSession, const FInternalLiveSessionCreationParams& AdditionalParams)
@@ -554,7 +558,7 @@ bool FConcertSyncServer::CreateLiveSession(const TSharedRef<IConcertServerSessio
 		// Create Replication Manager
 		if (EnumHasAnyFlags(LiveSession->GetSessionFlags(), EConcertSyncSessionFlags::EnableReplication))
 		{
-			CreateReplicationManager(InSession);
+			CreateReplicationManager(InSession, LiveSession->GetSessionFlags());
 		}
 
 		// We needn't call OnActivityProduced().Remove(...) because the subscription needs to stay for the lifetime of FConcertSyncServerLiveSession::SessionDatabase
diff --git a/Engine/Plugins/Developer/Concert/ConcertSync/ConcertSyncServer/Source/ConcertSyncServer/Private/ConcertSyncServer.h b/Engine/Plugins/Developer/Concert/ConcertSync/ConcertSyncServer/Source/ConcertSyncServer/Private/ConcertSyncServer.h
--- a/Engine/Plugins/Developer/Concert/ConcertSync/ConcertSyncServer/Source/ConcertSyncServer/Private/ConcertSyncServer.h
+++ b/Engine/Plugins/Developer/Concert/ConcertSync/ConcertSyncServer/Source/ConcertSyncServer/Private/ConcertSyncServer.h
@@ -61,8 +61,8 @@ private:
 	void CreateSequencerManager(const TSharedRef<FConcertSyncServerLiveSession>& InLiveSession);
 	void DestroySequencerManager(const TSharedRef<FConcertSyncServerLiveSession>& InLiveSession);
 	
-	void CreateReplicationManager(TSharedRef<IConcertServerSession> InLiveSession);
-	void DestroyReplicationManager(const TSharedRef<IConcertServerSession>& InLiveSession); 
+	void CreateReplicationManager(const TSharedRef<IConcertServerSession>& InSession, EConcertSyncSessionFlags InSessionFlags);
+	void DestroyReplicationManager(const FGuid& SessionId); 
 
 	bool CreateLiveSession(const TSharedRef<IConcertServerSession>& InSession, const FInternalLiveSessionCreationParams& AdditionalParams);
 	void DestroyLiveSession(const TSharedRef<IConcertServerSession>& InSession);
diff --git a/Engine/Plugins/Developer/Concert/ConcertSync/ConcertSyncServer/Source/ConcertSyncServer/Private/Replication/AuthorityManager.cpp b/Engine/Plugins/Developer/Concert/ConcertSync/ConcertSyncServer/Source/ConcertSyncServer/Private/Replication/AuthorityManager.cpp
--- a/Engine/Plugins/Developer/Concert/ConcertSync/ConcertSyncServer/Source/ConcertSyncServer/Private/Replication/AuthorityManager.cpp
+++ b/Engine/Plugins/Developer/Concert/ConcertSync/ConcertSyncServer/Source/ConcertSyncServer/Private/Replication/AuthorityManager.cpp
@@ -138,7 +138,7 @@ namespace UE::ConcertSyncServer::Replication
 		return EnumerateAuthorityConflicts(Object) == EAuthorityResult::Allowed;
 	}
 
-	void FAuthorityManager::OnClientLeft(const FClientId& ClientEndpointId)
+	void FAuthorityManager::OnPostClientLeft(const FClientId& ClientEndpointId)
 	{
 		ClientAuthorityData.Remove(ClientEndpointId);
 	}
diff --git a/Engine/Plugins/Developer/Concert/ConcertSync/ConcertSyncServer/Source/ConcertSyncServer/Private/Replication/AuthorityManager.h b/Engine/Plugins/Developer/Concert/ConcertSync/ConcertSyncServer/Source/ConcertSyncServer/Private/Replication/AuthorityManager.h
--- a/Engine/Plugins/Developer/Concert/ConcertSync/ConcertSyncServer/Source/ConcertSyncServer/Private/Replication/AuthorityManager.h
+++ b/Engine/Plugins/Developer/Concert/ConcertSync/ConcertSyncServer/Source/ConcertSyncServer/Private/Replication/AuthorityManager.h
@@ -74,7 +74,7 @@ namespace UE::ConcertSyncServer::Replication
 		bool CanTakeAuthority(const FConcertReplicatedObjectId& Object) const;
 
 		/** Notifies this manager that the client has left, which means all their authority is now gone. */
-		void OnClientLeft(const FClientId& ClientEndpointId);
+		void OnPostClientLeft(const FClientId& ClientEndpointId);
 		/** Takes away authority from the given client from the given object. */
 		void RemoveAuthority(const FConcertReplicatedObjectId& Object);
 		
diff --git a/Engine/Plugins/Developer/Concert/ConcertSync/ConcertSyncServer/Source/ConcertSyncServer/Private/Replication/ConcertServerReplicationManager.cpp b/Engine/Plugins/Developer/Concert/ConcertSync/ConcertSyncServer/Source/ConcertSyncServer/Private/Replication/ConcertServerReplicationManager.cpp
--- a/Engine/Plugins/Developer/Concert/ConcertSync/ConcertSyncServer/Source/ConcertSyncServer/Private/Replication/ConcertServerReplicationManager.cpp
+++ b/Engine/Plugins/Developer/Concert/ConcertSync/ConcertSyncServer/Source/ConcertSyncServer/Private/Replication/ConcertServerReplicationManager.cpp
@@ -23,11 +23,15 @@ namespace UE::ConcertSyncServer::Replication
 		TEXT("Whether to log changes to authority.")
 		);
 	
-	FConcertServerReplicationManager::FConcertServerReplicationManager(TSharedRef<IConcertServerSession> InLiveSession)
+	FConcertServerReplicationManager::FConcertServerReplicationManager(
+		TSharedRef<IConcertServerSession> InLiveSession,
+		EConcertSyncSessionFlags SessionFlags
+		)
 		: Session(MoveTemp(InLiveSession))
 		, ReplicationFormat(MakeUnique<ConcertSyncCore::FFullObjectFormat>())
 		, AuthorityManager(*this, Session)
-		, SyncControlManager(*Session, AuthorityManager, *this)
+		, MuteManager(*Session, ServerObjectCache, SessionFlags)
+		, SyncControlManager(*Session, AuthorityManager, MuteManager, *this)
 		, ReplicationCache(MakeShared<ConcertSyncCore::FObjectReplicationCache>(*ReplicationFormat))
 		, ReplicationDataReceiver(AuthorityManager, SyncControlManager, *Session, *ReplicationCache)
 	{
@@ -95,6 +99,7 @@ namespace UE::ConcertSyncServer::Replication
 		if (bSuccess)
 		{
 			Response.SyncControl = SyncControlManager.OnGenerateSyncControlForClientJoin(ClientId);
+			ServerObjectCache.OnJoin(ClientId, Request);
 		}
 		
 		UE_CLOG(bSuccess, LogConcert, Log, TEXT("Accepted replication join request"));
@@ -248,9 +253,24 @@ namespace UE::ConcertSyncServer::Replication
 
 	void FConcertServerReplicationManager::OnClientLeftReplication(const FGuid& EndpointId)
 	{
-		Clients.Remove(EndpointId);
-		AuthorityManager.OnClientLeft(EndpointId);
-		SyncControlManager.OnClientLeft(EndpointId);
+		if (!Clients.Contains(EndpointId))
+		{
+			return;
+		}
+		
+		TUniquePtr<FConcertReplicationClient> RemovedClient;
+		const bool bRemoved = Clients.RemoveAndCopyValue(EndpointId, RemovedClient);
+		check(bRemoved);
+
+		// ServerObjectCache should be updated before anyone else that may rely on its state.
+		ServerObjectCache.OnPostClientLeft(EndpointId, RemovedClient->GetStreamDescriptions());
+		
+		// There is some inefficiency here: FMuteManager::OnMuteStateChanged may broadcast, which causes SyncControlManager to rebuild ...
+		MuteManager.OnPostClientLeft(RemovedClient->GetStreamDescriptions());
+		AuthorityManager.OnPostClientLeft(EndpointId);
+		
+		// ... and then the sync control manager rebuilds again.
+		SyncControlManager.OnPostClientLeft(EndpointId);
 	}
 
 	void FConcertServerReplicationManager::Tick(IConcertServerSession& InSession, float InDeltaTime)
diff --git a/Engine/Plugins/Developer/Concert/ConcertSync/ConcertSyncServer/Source/ConcertSyncServer/Private/Replication/ConcertServerReplicationManager.h b/Engine/Plugins/Developer/Concert/ConcertSync/ConcertSyncServer/Source/ConcertSyncServer/Private/Replication/ConcertServerReplicationManager.h
--- a/Engine/Plugins/Developer/Concert/ConcertSync/ConcertSyncServer/Source/ConcertSyncServer/Private/Replication/ConcertServerReplicationManager.h
+++ b/Engine/Plugins/Developer/Concert/ConcertSync/ConcertSyncServer/Source/ConcertSyncServer/Private/Replication/ConcertServerReplicationManager.h
@@ -6,9 +6,11 @@
 #include "ConcertMessages.h"
 #include "ConcertReplicationClient.h"
 #include "Enumeration/IRegistrationEnumerator.h"
+#include "MuteManager.h"
 #include "Replication/Formats/IObjectReplicationFormat.h"
 #include "Replication/IConcertServerReplicationManager.h"
 #include "Replication/Messages/Handshake.h"
+#include "Replication/Misc/ReplicatedObjectHierarchyCache.h"
 #include "Replication/Processing/ObjectReplicationCache.h"
 #include "Replication/Processing/ServerObjectReplicationReceiver.h"
 #include "SyncControlManager.h"
@@ -20,6 +22,7 @@
 class IConcertClientReplicationBridge;
 class IConcertServerSession;
 
+enum class EConcertSyncSessionFlags : uint32;
 enum class EConcertQueryClientStreamFlags : uint8;
 
 struct FConcertReplication_ChangeStream_Response;
@@ -50,7 +53,7 @@ namespace UE::ConcertSyncServer::Replication
 	{
 	public:
 
-		explicit FConcertServerReplicationManager(TSharedRef<IConcertServerSession> InLiveSession);
+		explicit FConcertServerReplicationManager(TSharedRef<IConcertServerSession> InLiveSession, EConcertSyncSessionFlags SessionFlags);
 		virtual ~FConcertServerReplicationManager() override;
 
 		const FAuthorityManager& GetAuthorityManager() const { return AuthorityManager; }
@@ -71,8 +74,18 @@ namespace UE::ConcertSyncServer::Replication
 		/** Responsible for analysing received replication data. */
 		TUniquePtr<ConcertSyncCore::IObjectReplicationFormat> ReplicationFormat;
 
+		/**
+		 * Holds the outer hierarchy of all objects registered in any stream.
+		 * 
+		 * This receives join, leave, and change stream events.
+		 * The events are processed after the mute manager does.
+		 */
+		ConcertSyncCore::FReplicatedObjectHierarchyCache ServerObjectCache;
+
 		/** Responds to client requests to changing authority and can be asked whether an object change is valid to take place. */
 		FAuthorityManager AuthorityManager;
+		/** Responds to client mute requests and stores the mute states. */
+		FMuteManager MuteManager;
 		/** Decides whether clients should be replicating. Clients may replicate when they have authority and there are other clients listening for that data. */
 		FSyncControlManager SyncControlManager;
 		
diff --git a/Engine/Plugins/Developer/Concert/ConcertSync/ConcertSyncServer/Source/ConcertSyncServer/Private/Replication/ConcertServerReplicationManager_ChangeStream.cpp b/Engine/Plugins/Developer/Concert/ConcertSync/ConcertSyncServer/Source/ConcertSyncServer/Private/Replication/ConcertServerReplicationManager_ChangeStream.cpp
--- a/Engine/Plugins/Developer/Concert/ConcertSync/ConcertSyncServer/Source/ConcertSyncServer/Private/Replication/ConcertServerReplicationManager_ChangeStream.cpp
+++ b/Engine/Plugins/Developer/Concert/ConcertSync/ConcertSyncServer/Source/ConcertSyncServer/Private/Replication/ConcertServerReplicationManager_ChangeStream.cpp
@@ -157,15 +157,29 @@ namespace UE::ConcertSyncServer::Replication
 		const TUniquePtr<FConcertReplicationClient>* SendingClient = Clients.Find(SendingClientId);
 		if (SendingClient && Private::ShouldAcceptRequest(Request, *SendingClient->Get(), AuthorityManager, Response))
 		{
+			const TArray<FConcertReplicationStream>& StreamsBeforeChange = SendingClient->Get()->GetStreamDescriptions();
+
+			TArray<FConcertObjectInStreamID> AddedObjects;
+			ConcertSyncCore::Replication::ChangeStreamUtils::ForEachAddedObject(Request, StreamsBeforeChange, [&AddedObjects](const FConcertObjectInStreamID& Object)
+			{
+				AddedObjects.Add(Object);
+				return EBreakBehavior::Continue;
+			});
 			// If the client had authority over any objects that were removed by this request, authority must be cleaned up
-			ConcertSyncCore::Replication::ChangeStreamUtils::ForEachObjectLosingAuthority(Request, SendingClient->Get()->GetStreamDescriptions(),
-				[this, &SendingClientId](const FConcertObjectInStreamID& RemovedObject)
+			TArray<FConcertObjectInStreamID> RemovedObjects;
+			ConcertSyncCore::Replication::ChangeStreamUtils::ForEachRemovedObject(Request, StreamsBeforeChange,
+				[this, &SendingClientId, &RemovedObjects](const FConcertObjectInStreamID& RemovedObject)
 				{
 					AuthorityManager.RemoveAuthority({ RemovedObject, SendingClientId});
+					RemovedObjects.Add(RemovedObject);
 					return EBreakBehavior::Continue;
 				});
 			
+			// This needs to happen last since all of the above use ChangeStreamUtils::ForEachRemovedObject to determine the diff of changes.
 			SendingClient->Get()->ApplyValidatedRequest(Request);
+			
+			ServerObjectCache.OnChangeStreams(SendingClientId, AddedObjects, RemovedObjects);
+			MuteManager.PostApplyStreamChange(SendingClientId, AddedObjects, RemovedObjects);
 		}
 		else
 		{
diff --git a/Engine/Plugins/Developer/Concert/ConcertSync/ConcertSyncServer/Source/ConcertSyncServer/Private/Replication/MuteManager.cpp b/Engine/Plugins/Developer/Concert/ConcertSync/ConcertSyncServer/Source/ConcertSyncServer/Private/Replication/MuteManager.cpp
new file mode 100644
--- /dev/null
+++ b/Engine/Plugins/Developer/Concert/ConcertSync/ConcertSyncServer/Source/ConcertSyncServer/Private/Replication/MuteManager.cpp
@@ -0,0 +1,434 @@
+﻿// Copyright Epic Games, Inc. All Rights Reserved.
+
+#include "MuteManager.h"
+
+#include "ConcertLogGlobal.h"
+#include "IConcertSession.h"
+#include "Misc/ObjectPathOuterIterator.h"
+#include "Replication/ChangeStreamSharedUtils.h"
+#include "Replication/Messages/Muting.h"
+#include "Replication/Misc/ReplicatedObjectHierarchyCache.h"
+#include "Util/LogUtils.h"
+
+#include "HAL/IConsoleManager.h"
+#include "Misc/ScopeExit.h"
+
+namespace UE::ConcertSyncServer::Replication
+{
+	TAutoConsoleVariable<bool> CVarLogMuteRequestsAndResponsesOnServer(
+		TEXT("Concert.Replication.LogMuteRequestsAndResponsesOnServer"),
+		false,
+		TEXT("Whether to log changes to the mute state.")
+		);
+	
+	FMuteManager::FMuteManager(
+		IConcertSession& InSession,
+		const ConcertSyncCore::FReplicatedObjectHierarchyCache& ServerObjectCache,
+		EConcertSyncSessionFlags InSessionFlags
+		)
+		: Session(InSession)
+		, ServerObjectCache(ServerObjectCache)
+		, SessionFlags(InSessionFlags)
+	{
+		Session.RegisterCustomRequestHandler<FConcertReplication_ChangeMuteState_Request, FConcertReplication_ChangeMuteState_Response>(this, &FMuteManager::HandleChangeMuteStateRequest);
+		Session.RegisterCustomRequestHandler<FConcertReplication_QueryMuteState_Request, FConcertReplication_QueryMuteState_Response>(this, &FMuteManager::HandleQueryMuteStateRequest);
+	}
+
+	FMuteManager::~FMuteManager()
+	{
+		Session.UnregisterCustomRequestHandler<FConcertReplication_ChangeMuteState_Request>();
+		Session.UnregisterCustomRequestHandler<FConcertReplication_QueryMuteState_Request>();
+	}
+
+	bool FMuteManager::IsMuted(const FSoftObjectPath& Object) const
+	{
+		const FMuteData* ObjectMuteState = MuteState.Find(Object);
+		return ObjectMuteState
+			&& (ObjectMuteState->State == EMuteState::ExplicitlyMuted || ObjectMuteState->State == EMuteState::ImplicitlyMuted);
+	}
+
+	void FMuteManager::PostApplyStreamChange(const FGuid& ClientId, TConstArrayView<FConcertObjectInStreamID> AddedObjects, TConstArrayView<FConcertObjectInStreamID> RemovedObjects)
+	{
+		// If Request adds new subobjects that should be muted implicitly by existing rules, update MuteState.
+		TrackAddedSubobjectsForImplicitMuting(AddedObjects);
+		// Objects that are no longer referenced by any clients must be removed (to avoid leaks)
+		UnmuteObjectsIfUnreferenced(RemovedObjects);
+		
+		// This updates sync control, which may remove sync control from other clients.
+		check(OnMuteStateChangedDelegate.IsBound());
+		OnMuteStateChangedDelegate.Execute(ClientId);
+	}
+
+	void FMuteManager::OnPostClientLeft(const TArray<FConcertReplicationStream>& ClientStreams)
+	{
+		// We'll treat this case as a FConcertReplication_ChangeStream_Request that removes all objects 
+		TSet<FSoftObjectPath> VisitedObjects;
+		for (const FConcertReplicationStream& Stream : ClientStreams)
+		{
+			for (const TPair<FSoftObjectPath, FConcertReplicatedObjectInfo>& ObjectPair : Stream.BaseDescription.ReplicationMap.ReplicatedObjects)
+			{
+				const FSoftObjectPath ObjectPath = ObjectPair.Key; 
+				if (!VisitedObjects.Contains(ObjectPath))
+				{
+					VisitedObjects.Add(ObjectPath);
+					UnmuteObjectIfUnreferenced(ObjectPath);
+				}
+			}
+		}
+	}
+	
+	void FMuteManager::TrackAddedSubobjectsForImplicitMuting(TConstArrayView<FConcertObjectInStreamID> AddedObjects)
+	{
+		for (const FConcertObjectInStreamID& Object : AddedObjects)
+		{
+			const FSoftObjectPath ObjectPath = Object.Object;
+			if (MuteState.Contains(ObjectPath))
+			{
+				continue;
+			}
+
+			const auto[_, ParentMuteData] = FindAffectingParentState(ObjectPath);
+			if (ParentMuteData && ParentMuteData->State == EMuteState::ExplicitlyUnmuted)
+			{
+				MuteState.Add(ObjectPath).State = EMuteState::ImplicitlyMuted;
+			}
+			else if (ParentMuteData && ParentMuteData->State == EMuteState::ExplicitlyMuted)
+			{
+				MuteState.Add(ObjectPath).State = EMuteState::ImplicitlyUnmuted;
+			}
+		}
+	}
+
+	void FMuteManager::UnmuteObjectsIfUnreferenced(TConstArrayView<FConcertObjectInStreamID> RemovedObjects)
+	{
+		// We'll avoid processing any object again that we've already removed... not doing so would unnecessarily iterate the hierarchy.
+		TSet<FSoftObjectPath> ProcessedObjects;
+		
+		for (const FConcertObjectInStreamID& RemovedObject : RemovedObjects)
+		{
+			const FSoftObjectPath& RemovedObjectPath = RemovedObject.Object;
+			if (!ProcessedObjects.Contains(RemovedObjectPath))
+			{
+				UnmuteObjectIfUnreferenced(RemovedObjectPath, [&ProcessedObjects](const FSoftObjectPath& ObjectPath)
+				{
+					ProcessedObjects.Add(ObjectPath);
+				});
+			}
+		}
+	}
+
+	template<CObjectProcessable TCallback>
+	void FMuteManager::UnmuteObjectIfUnreferenced(const FSoftObjectPath& RemovedObject, TCallback&& OnRemoved)
+	{
+		// If RemovedObject is referenced or has known subobjects, the mute state should tick around.
+		const bool bHasChildrenOrIsReferenced = ServerObjectCache.HasChildren(RemovedObject);
+		if (bHasChildrenOrIsReferenced)
+		{
+			return;
+		}
+		
+		const auto [AffectingParent, ParentMuteData] = FindAffectingParentState(RemovedObject);
+		// RemovedObject is not referenced by any streams, so it has no business being muted.
+		RemoveMuteState(RemovedObject, OnRemoved);
+
+		// There may be a parent object, which also no longer has known children ... so remove that as well.
+		if (ParentMuteData)
+		{
+			UnmuteObjectIfUnreferenced(AffectingParent, OnRemoved);
+		}
+	}
+
+	EConcertSessionResponseCode FMuteManager::HandleQueryMuteStateRequest(
+		const FConcertSessionContext& Context,
+		const FConcertReplication_QueryMuteState_Request& Request,
+		FConcertReplication_QueryMuteState_Response& Response
+		)
+	{
+		if (!EnumHasAnyFlags(SessionFlags, EConcertSyncSessionFlags::ShouldAllowGlobalMuting))
+		{
+			return EConcertSessionResponseCode::Failed;
+		}
+
+		Response.ErrorCode = EReplicationResponseErrorCode::Handled;
+		const auto AddObjectToResponse = [this, &Response](const FSoftObjectPath& ObjectPath, const FMuteData& ObjectMuteData)
+		{
+			switch (ObjectMuteData.State)
+			{
+			case EMuteState::ExplicitlyMuted:
+				Response.ExplicitlyMutedObjects.Add(ObjectPath, *ObjectMuteData.MuteSetting);
+				break;
+			case EMuteState::ExplicitlyUnmuted:
+				Response.ExplicitlyUnmutedObjects.Add(ObjectPath, *ObjectMuteData.MuteSetting);
+				break;
+			case EMuteState::ImplicitlyMuted:
+				Response.ImplicitlyMutedObjects.Add(ObjectPath);
+				break;
+			case EMuteState::ImplicitlyUnmuted:
+				Response.ImplicitlyUnmutedObjects.Add(ObjectPath);
+				break;
+			}
+		};
+
+		if (Request.WantsAllObjects())
+		{
+			for (const TPair<FSoftObjectPath, FMuteData>& Pair : MuteState)
+			{
+				AddObjectToResponse(Pair.Key, Pair.Value);
+			}
+		}
+		else
+		{
+			for (const FSoftObjectPath& QueriedPath : Request.QueriedObjects)
+			{
+				if (const FMuteData* ObjectMuteData = MuteState.Find(QueriedPath))
+				{
+					AddObjectToResponse(QueriedPath, *ObjectMuteData);
+				}
+			}
+		}
+		
+		return EConcertSessionResponseCode::Success;
+	}
+
+	EConcertSessionResponseCode FMuteManager::HandleChangeMuteStateRequest(
+		const FConcertSessionContext& Context,
+		const FConcertReplication_ChangeMuteState_Request& Request,
+		FConcertReplication_ChangeMuteState_Response& Response
+		)
+	{
+		LogNetworkMessage(CVarLogMuteRequestsAndResponsesOnServer, Request, [this, &Context](){ return GetClientName(Session, Context.SourceEndpointId); });
+		ON_SCOPE_EXIT{ LogNetworkMessage(CVarLogMuteRequestsAndResponsesOnServer, Response, [this, &Context](){ return GetClientName(Session, Context.SourceEndpointId); }); };
+		
+		if (!EnumHasAnyFlags(SessionFlags, EConcertSyncSessionFlags::ShouldAllowGlobalMuting))
+		{
+			return EConcertSessionResponseCode::Failed;
+		}
+
+		if (!ValidateRequest(Request, Response))
+		{
+			Response.ErrorCode = EConcertReplicationMuteErrorCode::Rejected;
+			return EConcertSessionResponseCode::Success;
+		}
+
+		Response.ErrorCode = EConcertReplicationMuteErrorCode::Accepted;
+		ApplyRequest(Request);
+		
+		// Sync control will now 1. generate a response we can embed in the response, and 2. send a network event to all other clients.
+		// The embedded response will contain only those objects which changed sync control
+		// Fyi, the new sync control may contain more objects than were specified in the request since parent objects can also cause their subobjects' mute state to change.
+		check(OnGenerateSyncControlForMuteChangeDelegate.IsBound());
+		Response.SyncControl = OnGenerateSyncControlForMuteChangeDelegate.Execute(Context.SourceEndpointId);
+#if UE_BUILD_DEBUG // In debug builds so we don't tank performance
+		checkf(!Response.SyncControl.DoesAtLeastOneObjectLoseSyncControl(), TEXT("Clients are supposed to infer that objects lose sync control. We only tell them which objects gained sync control."));
+#endif
+		
+		return EConcertSessionResponseCode::Success;
+	}
+
+	bool FMuteManager::ValidateRequest(const FConcertReplication_ChangeMuteState_Request& Request, FConcertReplication_ChangeMuteState_Response& Response) const
+	{
+		bool bIsValidRequest = true;
+		const auto Reject = [&Response, &bIsValidRequest](const FSoftObjectPath& ObjectPath){ Response.RejectionReasons.Add(ObjectPath); bIsValidRequest = false; };
+		const auto SharedValidateSetting = [this, &Reject](const FSoftObjectPath& Object, const FConcertReplication_ObjectMuteSetting& Setting)
+		{
+			const bool bIsObjectReferenced = ServerObjectCache.IsObjectReferencedDirectly(Object);
+			const bool bAppliesToSubobjects = ConcertSyncCore::AffectSubobjects(Setting.Flags);
+			if (!bIsObjectReferenced && !bAppliesToSubobjects)
+			{
+				Reject(Object);
+				return;
+			}
+
+			const bool bHasSubobjects = ServerObjectCache.HasChildren(Object);
+			if (!bIsObjectReferenced && bAppliesToSubobjects && !bHasSubobjects)
+			{
+				Reject(Object);
+			}
+		};
+		
+		for (const TPair<FSoftObjectPath, FConcertReplication_ObjectMuteSetting>& ToMute : Request.ObjectsToMute)
+		{
+			SharedValidateSetting(ToMute.Key, ToMute.Value);
+
+			// Cannot mute & unmute at the same time
+			if (Request.ObjectsToUnmute.Contains(ToMute.Key))
+			{
+				Reject(ToMute.Key);
+			}
+		}
+		for (const TPair<FSoftObjectPath, FConcertReplication_ObjectMuteSetting>& ToUnmute : Request.ObjectsToUnmute)
+		{
+			// If something is muted, the user can always unmute it. If it is already unmuted, it's still valid (but a no-op).
+			if (!MuteState.Contains(ToUnmute.Key))
+			{
+				SharedValidateSetting(ToUnmute.Key, ToUnmute.Value);
+			}
+		}
+		
+		return bIsValidRequest;
+	}
+
+	void FMuteManager::ApplyRequest(const FConcertReplication_ChangeMuteState_Request& Request)
+	{
+		UE_LOG(LogConcert, Log, TEXT("Explicitly muting %d and unmuting %d objects."), Request.ObjectsToMute.Num(), Request.ObjectsToUnmute.Num());
+		
+		for (const TPair<FSoftObjectPath, FConcertReplication_ObjectMuteSetting>& ToMute : Request.ObjectsToMute)
+		{
+			const FSoftObjectPath& ObjectToMute = ToMute.Key;
+			const FConcertReplication_ObjectMuteSetting& MuteSetting = ToMute.Value;
+
+			// Muting object again with same parameters is a valid: not work to do if so.
+			const bool bAffectsChildren = ConcertSyncCore::AffectSubobjects(MuteSetting.Flags);
+			{
+				const FMuteData* ObjectMuteState = MuteState.Find(ObjectToMute);
+				const bool bAlreadyAtSetting = ObjectMuteState && ObjectMuteState->MuteSetting && ObjectMuteState->MuteSetting == ToMute.Value;
+				if (bAlreadyAtSetting)
+				{
+					continue;
+				}
+
+				// If the request is changing the flag from IncludeSubobjects -> None, the children need their implicit state updated.
+				const bool bPreviousAffectedChildren = ObjectMuteState && ObjectMuteState->AffectsSubobjects(); 
+				if (!bAffectsChildren && bPreviousAffectedChildren)
+				{
+					RemoveMuteState(ObjectToMute);
+				}
+				// ObjectMuteState may now be invalid due to possible removal
+			}
+
+			const FMuteData NewMuteData
+			{
+				MuteSetting,
+				EMuteState::ExplicitlyMuted
+			};
+			MuteState.Add(ObjectToMute, NewMuteData);
+
+			// Children may need their implicit state updated
+			if (bAffectsChildren)
+			{
+				UpdateImplicitStateUnder(ObjectToMute, EMuteState::ImplicitlyMuted);
+			}
+		}
+
+		for (const TPair<FSoftObjectPath, FConcertReplication_ObjectMuteSetting>& ToMute : Request.ObjectsToUnmute)
+		{
+			const FSoftObjectPath& ObjectToUnmute = ToMute.Key;
+			const FConcertReplication_ObjectMuteSetting& MuteSetting = ToMute.Value;
+			
+			const auto[AffectingParent, ParentMuteState] = FindAffectingParentState(ObjectToUnmute);
+			// ParentMuteState can either be:
+			//	- EMuteState::ExplicitlyMuted: all children should be muted
+			//	- EMuteState::ExplicitlyUnmuted: by invariant, there must be a parent that is ExplicitlyMuted with EConcertReplicationMuteFlags::IncludeSubobjects.
+			// Hence in both cases, children need to be updated to be marked as implicitly muted.
+			if (ParentMuteState)
+			{
+				MuteState.Add(ObjectToUnmute, { MuteSetting, EMuteState::ExplicitlyUnmuted });
+				if (ConcertSyncCore::AffectSubobjects(MuteSetting.Flags))
+				{
+					UpdateImplicitStateUnder(ObjectToUnmute, EMuteState::ImplicitlyUnmuted);
+				}
+			}
+			else
+			{
+				// If no parent state affects ObjectToUnmute, the entire subtree must be removed.
+				FMuteData MuteData;
+				const bool bHadData = MuteState.RemoveAndCopyValue(ObjectToUnmute, MuteData);
+				if (bHadData && MuteData.AffectsSubobjects())
+				{
+					ClearAllChildStateUnder(ObjectToUnmute);
+				}
+			}
+		}
+	}
+
+	template<CObjectProcessable TCallback>
+	void FMuteManager::RemoveMuteState(const FSoftObjectPath& Object, TCallback&& OnRemoved)
+	{
+		const FMuteData* ObjectMuteData = MuteState.Find(Object);
+		if (!ObjectMuteData)
+		{
+			return;
+		}
+		
+		// If an Object affects its subobjects...
+		if (ObjectMuteData->AffectsSubobjects())
+		{
+			const auto[AffectingParent, ParentMuteData] = FindAffectingParentState(Object);
+			if (ParentMuteData)
+			{
+				// ... and has a parent also affecting subobjects, then all of Object's children are now affected by that parent.
+				const EMuteState ImplicitNewChildState = ParentMuteData->State == EMuteState::ExplicitlyMuted ? EMuteState::ImplicitlyMuted : EMuteState::ImplicitlyUnmuted;
+				UpdateImplicitStateUnder(Object, ImplicitNewChildState);
+			}
+			else
+			{
+				// ... and does not have any parent affecting subobjects, then clear all child state
+				ClearAllChildStateUnder(Object, OnRemoved);
+			}
+		}
+
+		MuteState.Remove(Object);
+		OnRemoved(Object);
+	}
+
+	void FMuteManager::UpdateImplicitStateUnder(const FSoftObjectPath& Parent, EMuteState NewImplicitState)
+	{
+		check(NewImplicitState == EMuteState::ImplicitlyMuted || NewImplicitState == EMuteState::ImplicitlyUnmuted);
+		
+		ServerObjectCache.TraverseTopToBottom([this, NewImplicitState](const ConcertSyncCore::FChildRelation& Relation)
+		{
+			FMuteData* SubobjectMuteState = MuteState.Find(Relation.Child.Object);
+			if (!SubobjectMuteState)
+			{
+				MuteState.Add(Relation.Child.Object, { {}, NewImplicitState });
+			}
+			else if (SubobjectMuteState->MuteSetting
+				&& ConcertSyncCore::AffectSubobjects(SubobjectMuteState->MuteSetting->Flags))
+			{
+				return ConcertSyncCore::ETreeTraversalBehavior::SkipSubtree;
+			}
+			else
+			{
+				SubobjectMuteState->State = NewImplicitState;
+			}
+					
+			return ConcertSyncCore::ETreeTraversalBehavior::Continue;
+		}, Parent);
+	}
+	
+	template<CObjectProcessable TCallback>
+	void FMuteManager::ClearAllChildStateUnder(const FSoftObjectPath& ParentObject, TCallback&& OnRemoved)
+	{
+		ServerObjectCache.TraverseTopToBottom([this, &OnRemoved](const ConcertSyncCore::FChildRelation& Relation)
+		{
+			FMuteData* SubobjectMuteState = MuteState.Find(Relation.Child.Object);
+			if (!SubobjectMuteState)
+			{
+				return ConcertSyncCore::ETreeTraversalBehavior::Continue;
+			}
+
+			MuteState.Remove(Relation.Child.Object);
+			OnRemoved(Relation.Child.Object);
+			return ConcertSyncCore::ETreeTraversalBehavior::Continue;
+		}, ParentObject);
+	}
+	
+	TPair<FSoftObjectPath, const FMuteManager::FMuteData*> FMuteManager::FindAffectingParentState(const FSoftObjectPath& Subobject) const
+	{
+		for (ConcertSyncCore::FObjectPathOuterIterator It(Subobject); It; ++It)
+		{
+			const FSoftObjectPath& Parent = *It;
+			if (const FMuteData* ParentMuteState = MuteState.Find(Parent)
+				; ParentMuteState
+				// Does it have the IncludeSubobjects flag
+				&& ParentMuteState->MuteSetting && ConcertSyncCore::AffectSubobjects(ParentMuteState->MuteSetting->Flags)
+				// Only explicit objects should have the IncludeSubobjects flag
+				&& ensure(ParentMuteState->IsExplicit()))
+			{
+				return { Parent, ParentMuteState };
+			}
+		}
+		
+		return {};
+	}
+}
diff --git a/Engine/Plugins/Developer/Concert/ConcertSync/ConcertSyncServer/Source/ConcertSyncServer/Private/Replication/MuteManager.h b/Engine/Plugins/Developer/Concert/ConcertSync/ConcertSyncServer/Source/ConcertSyncServer/Private/Replication/MuteManager.h
new file mode 100644
--- /dev/null
+++ b/Engine/Plugins/Developer/Concert/ConcertSync/ConcertSyncServer/Source/ConcertSyncServer/Private/Replication/MuteManager.h
@@ -0,0 +1,154 @@
+﻿// Copyright Epic Games, Inc. All Rights Reserved.
+
+#pragma once
+
+#include "ConcertMessages.h"
+#include "ConcertSyncSessionFlags.h"
+#include "IConcertSessionHandler.h"
+#include "Replication/Messages/Muting.h"
+#include "Templates/UnrealTemplate.h"
+
+#include <type_traits>
+
+class IConcertSession;
+struct FConcertReplication_ChangeStream_Request;
+struct FConcertReplicationStream;
+struct FConcertObjectInStreamID;
+
+namespace UE::ConcertSyncCore { class FReplicatedObjectHierarchyCache; }
+namespace UE::ConcertSyncServer::Replication { class IRegistrationEnumerator; }
+
+namespace UE::ConcertSyncServer::Replication
+{
+	/** Invocable with signature void([](const FConcertObjectInStreamID&){}) which calls the passed in lambda */
+	template<typename TEnumerateObjects>
+	concept CObjectEnumeratable = requires(TEnumerateObjects EnumerateObjects)
+	{
+		EnumerateObjects([](const FConcertObjectInStreamID&){});
+	};
+	
+	/** Invocable with signature void(FSoftObjectPath) */
+	template<typename TLambda>
+	concept CObjectProcessable = std::is_invocable_v<TLambda, const FSoftObjectPath&>;
+	
+	/**
+	 * Manages the session's replication mute state.
+	 * Handles FConcertReplication_ChangeMuteState_Request and FConcertReplication_QueryMuteState_Request requets.
+	 */
+	class FMuteManager : public FNoncopyable
+	{
+	public:
+
+		/** Updates sync control if needed in response to a client changing their mute state indirectly (e.g. because client removed object from their stream). */
+		DECLARE_DELEGATE_OneParam(FOnMuteStateChangedByClient, const FGuid& ClientId);
+		/** Updates sync control for all clients, sends an update to all clients but ClientId, and returns the sync control to embed into the mute response. */
+		DECLARE_DELEGATE_RetVal_OneParam(FConcertReplication_ChangeSyncControl, FGenerateSyncControlForMuteChange, const FGuid& ClientId);
+
+		FMuteManager(
+			IConcertSession& InSession UE_LIFETIMEBOUND,
+			const ConcertSyncCore::FReplicatedObjectHierarchyCache& ServerObjectCache UE_LIFETIMEBOUND,
+			EConcertSyncSessionFlags InSessionFlags
+			);
+		~FMuteManager();
+
+		/** @return Whether the object is globally muted. */
+		bool IsMuted(const FSoftObjectPath& Object) const;
+		
+		/** Called right after objects have been unregistered from ClientId's streams. */
+		void PostApplyStreamChange(const FGuid& ClientId, TConstArrayView<FConcertObjectInStreamID> AddedObjects, TConstArrayView<FConcertObjectInStreamID> RemovedObjects);
+		/** Called by FConcertServerReplicationManager when client leaves replication. */
+		void OnPostClientLeft(const TArray<FConcertReplicationStream>& ClientStreams);
+
+		FOnMuteStateChangedByClient& OnUpdateSyncControlForIndirectMuteChange() { return OnMuteStateChangedDelegate; }
+		FGenerateSyncControlForMuteChange& OnGenerateSyncControlForMuteChange() { return OnGenerateSyncControlForMuteChangeDelegate; }
+		
+	private:
+
+		/** The session to receive requests on. */
+		IConcertSession& Session;
+		/** Keeps track of all objects that are in client streams. Allows efficient traversing of subobject hierarchy. */
+		const ConcertSyncCore::FReplicatedObjectHierarchyCache& ServerObjectCache;
+
+		/** If ShouldAllowGlobalMuting is not set, the requests are not served. */
+		const EConcertSyncSessionFlags SessionFlags;
+
+		/** Broadcasts after the mute state has been changed. */
+		FOnMuteStateChangedByClient OnMuteStateChangedDelegate;
+		/**
+		 * Delegate into FSyncControlManager
+		 * 
+		 * Updates sync control for all clients, sends an update to all clients but ClientId, and returns the sync control to embed into the mute response.
+		 * The sync control will only contain RestrictToObjects.
+		 */
+		FGenerateSyncControlForMuteChange OnGenerateSyncControlForMuteChangeDelegate;
+
+		enum class EMuteState : uint8
+		{
+			/** Object is muted */
+			ExplicitlyMuted,
+			/** A parent object is ExplicitlyMuted and has EConcertReplicationMuteFlags::IncludeSubobjects set but is object and its children are supposed to be implicitly unmuted.  */
+			ExplicitlyUnmuted,
+			
+			/** Object is muted because one of its parent objects is ExplicitlyMuted and EConcertReplicationMuteFlags::IncludeSubobjects set. */
+			ImplicitlyMuted,
+			/** Object is unmuted because one of its parent objects is ExplicitlyUnmuted and has EConcertReplicationMuteFlags::IncludeSubobjects set. */
+			ImplicitlyUnmuted,
+		};
+		
+		struct FMuteData
+		{
+			/** Only set if the object was explicitly muted or unmuted. */
+			TOptional<FConcertReplication_ObjectMuteSetting> MuteSetting;
+			
+			EMuteState State;
+
+			bool IsExplicit() const { return MuteSetting.IsSet(); }
+			bool IsImplicit() const { return !IsExplicit(); }
+			bool AffectsSubobjects() const { return MuteSetting && ConcertSyncCore::AffectSubobjects(MuteSetting->Flags); }
+		};
+		/**
+		 * Holds the mute state of all objects.
+		 * 
+		 * Objects only show up in this map if they are either explicitly or implicitly affected by a mute setting:
+		 * - explicit means that a request listed them to be muted / unmuted
+		 * - implicit means that a parent object (i.e. an outer) is explicitly muted or unmuted.
+		 * Hence, if an object is not in this map it means that neither it nor any of its parents are explicitly/implicitly muted nor unmuted.
+		 */
+		TMap<FSoftObjectPath, FMuteData> MuteState;
+		
+		/** For each added object, checks whether it is a subobject of an object that is muted with the EConcertReplicationMuteFlags::IncludeSubobjects setting. */
+		void TrackAddedSubobjectsForImplicitMuting(TConstArrayView<FConcertObjectInStreamID> AddedObjects);
+		/** Calls RemoveMuteState on every object that is unreferenced by client streams (directly or indirectly via a subobject). */
+		void UnmuteObjectsIfUnreferenced(TConstArrayView<FConcertObjectInStreamID> RemovedObjects);
+		/** Calls RemoveMuteState if RemovedObject is unreferenced by client streams (directly or indirectly via a subobject). */
+		template<CObjectProcessable TCallback>
+		void UnmuteObjectIfUnreferenced(const FSoftObjectPath& RemovedObject, TCallback&& OnRemoved);
+		void UnmuteObjectIfUnreferenced(const FSoftObjectPath& RemovedObject) { UnmuteObjectIfUnreferenced(RemovedObject, [](const FSoftObjectPath& ObjectPath){}); }
+
+		/** Handles requests to query mute states. */
+		EConcertSessionResponseCode HandleQueryMuteStateRequest(const FConcertSessionContext& Context, const FConcertReplication_QueryMuteState_Request& Request, FConcertReplication_QueryMuteState_Response& Response);
+		/** Handles requests to change the mute states */
+		EConcertSessionResponseCode HandleChangeMuteStateRequest(const FConcertSessionContext& Context, const FConcertReplication_ChangeMuteState_Request& Request, FConcertReplication_ChangeMuteState_Response& Response);
+		/** Checks whether Request is valid to apply. */
+		bool ValidateRequest(const FConcertReplication_ChangeMuteState_Request& Request, FConcertReplication_ChangeMuteState_Response& Response) const;
+		/** Updates the internal state (assuming it is a valid request - call ValidateRequest before). */
+		void ApplyRequest(const FConcertReplication_ChangeMuteState_Request& Request);
+
+		/** Removes all mute state for Object and transitively updates subobjects if applicable. */
+		template<CObjectProcessable TCallback>
+		void RemoveMuteState(const FSoftObjectPath& Object, TCallback&& OnRemoved);
+		void RemoveMuteState(const FSoftObjectPath& Object) { RemoveMuteState(Object, [](const FSoftObjectPath&){}); } 
+		
+		/** Sets the mute state of all subobjects of Parent until reaching a subobject that also has EConcertReplicationMuteFlags::IncludeSubobjects set. */
+		void UpdateImplicitStateUnder(const FSoftObjectPath& Parent, EMuteState NewImplicitState);
+		
+		/** Removes all subobject mute state affected by Parent's mute state assuming that ParentObject has no parents with the EConcertReplicationMuteFlags::IncludeSubobjects. */
+		template<CObjectProcessable TCallback>
+		void ClearAllChildStateUnder(const FSoftObjectPath& ParentObject, TCallback&& OnRemoved);
+		void ClearAllChildStateUnder(const FSoftObjectPath& ParentObject) { ClearAllChildStateUnder(ParentObject, [](const FSoftObjectPath&){}); }
+		
+		/** Walks up the object hierarchy and returns the parent that affects ObjectPath, if any. */
+		TPair<FSoftObjectPath, const FMuteData*> FindAffectingParentState(const FSoftObjectPath& Subobject) const;
+	};
+}
+
diff --git a/Engine/Plugins/Developer/Concert/ConcertSync/ConcertSyncServer/Source/ConcertSyncServer/Private/Replication/ReplicationTestInterface.cpp b/Engine/Plugins/Developer/Concert/ConcertSync/ConcertSyncServer/Source/ConcertSyncServer/Private/Replication/ReplicationTestInterface.cpp
--- a/Engine/Plugins/Developer/Concert/ConcertSync/ConcertSyncServer/Source/ConcertSyncServer/Private/Replication/ReplicationTestInterface.cpp
+++ b/Engine/Plugins/Developer/Concert/ConcertSync/ConcertSyncServer/Source/ConcertSyncServer/Private/Replication/ReplicationTestInterface.cpp
@@ -14,9 +14,10 @@ class IConcertServerSession;
 namespace UE::ConcertSyncServer::TestInterface
 {
 	CONCERTSYNCSERVER_API TSharedRef<Replication::IConcertServerReplicationManager> CreateServerReplicationManager(
-		TSharedRef<IConcertServerSession> InLiveSession
+		TSharedRef<IConcertServerSession> InLiveSession,
+		EConcertSyncSessionFlags InSessionFlags
 		)
 	{
-		return MakeShared<Replication::FConcertServerReplicationManager>(MoveTemp(InLiveSession));
+		return MakeShared<Replication::FConcertServerReplicationManager>(MoveTemp(InLiveSession), InSessionFlags);
 	}
 }
diff --git a/Engine/Plugins/Developer/Concert/ConcertSync/ConcertSyncServer/Source/ConcertSyncServer/Private/Replication/SyncControlManager.cpp b/Engine/Plugins/Developer/Concert/ConcertSync/ConcertSyncServer/Source/ConcertSyncServer/Private/Replication/SyncControlManager.cpp
--- a/Engine/Plugins/Developer/Concert/ConcertSync/ConcertSyncServer/Source/ConcertSyncServer/Private/Replication/SyncControlManager.cpp
+++ b/Engine/Plugins/Developer/Concert/ConcertSync/ConcertSyncServer/Source/ConcertSyncServer/Private/Replication/SyncControlManager.cpp
@@ -3,6 +3,7 @@
 #include "SyncControlManager.h"
 
 #include "AuthorityManager.h"
+#include "MuteManager.h"
 #include "ConcertServer/Private/ConcertServerSession.h"
 #include "Replication/Data/ReplicationStream.h"
 #include "Replication/Messages/SyncControl.h"
@@ -12,19 +13,35 @@ namespace UE::ConcertSyncServer::Replication
 	FSyncControlManager::FSyncControlManager(
 		IConcertServerSession& Session,
 		FAuthorityManager& AuthorityManager,
+		FMuteManager& MuteManager,
 		const IRegistrationEnumerator& Getters
 		)
 		: Session(Session)
 		, AuthorityManager(AuthorityManager)
+		, MuteManager(MuteManager)
 		, Getters(Getters)
 	{
 		AuthorityManager.OnGenerateSyncControl().BindRaw(this, &FSyncControlManager::OnGenerateSyncControlForAuthorityResponse);
+		MuteManager.OnUpdateSyncControlForIndirectMuteChange().BindRaw(this, &FSyncControlManager::OnUpdateSyncControlForIndirectMuteChange);
+		MuteManager.OnGenerateSyncControlForMuteChange().BindRaw(this, &FSyncControlManager::OnGenerateSyncControlForClientMuteChange);
 	}
 
 	FSyncControlManager::~FSyncControlManager()
 	{
 		Session.OnSessionClientChanged().RemoveAll(this);
-		AuthorityManager.OnGenerateSyncControl().Unbind();
+		
+		if (ensure(AuthorityManager.OnGenerateSyncControl().IsBoundToObject(this)))
+		{
+			AuthorityManager.OnGenerateSyncControl().Unbind();
+		}
+		if (ensure(MuteManager.OnUpdateSyncControlForIndirectMuteChange().IsBoundToObject(this)))
+		{
+			MuteManager.OnUpdateSyncControlForIndirectMuteChange().Unbind();
+		}
+		if (ensure(MuteManager.OnGenerateSyncControlForMuteChange().IsBoundToObject(this)))
+		{
+			MuteManager.OnGenerateSyncControlForMuteChange().Unbind();
+		}
 	}
 
 	bool FSyncControlManager::HasSyncControl(const FConcertReplicatedObjectId& Object) const
@@ -42,7 +59,6 @@ namespace UE::ConcertSyncServer::Replication
 	FConcertReplication_ChangeSyncControl FSyncControlManager::OnGenerateSyncControlForAuthorityResponse(const FGuid& ClientId)
 	{
 		RefreshAndSendToAllClientsExcept(ClientId);
-
 		const auto ShouldSkipInControlMessage = [](const FConcertObjectInStreamID&, bool bNewState)
 		{
 			// We skip including states that disable sync control to save on network bandwidth.
@@ -53,6 +69,19 @@ namespace UE::ConcertSyncServer::Replication
 		return RefreshClientSyncControl(ClientId, ShouldSkipInControlMessage);
 	}
 
+	FConcertReplication_ChangeSyncControl FSyncControlManager::OnGenerateSyncControlForClientMuteChange(const FGuid& ClientId)
+	{
+		RefreshAndSendToAllClientsExcept(ClientId);
+		const auto ShouldSkipInControlMessage = [](const FConcertObjectInStreamID&, bool bNewState)
+		{
+			// We skip including states that disable sync control to save on network bandwidth.
+			// The receiving client can already predict the objects that lose sync control based on their request.
+			// See FConcertReplication_ChangeMuteState_Request and FConcertReplication_ChangeMuteState_Response::SyncControl documentation.
+			return !bNewState;
+		};
+		return RefreshClientSyncControl(ClientId, ShouldSkipInControlMessage);
+	}
+
 	void FSyncControlManager::HandleClientLeave(const FGuid& LeftClientId)
 	{
 		// The disconnecting client's objects should be removed
@@ -91,6 +120,7 @@ namespace UE::ConcertSyncServer::Replication
 					const FConcertReplicatedObjectId ReplicatedObjectId { ObjectId, ClientId };
 					const bool bIsValidForSending = Pair.Value.IsValidForSendingToServer()
 						&& AuthorityManager.HasAuthorityToChange(ReplicatedObjectId)
+						&& !MuteManager.IsMuted(ReplicatedObjectId.Object)
 						&& IsAnyoneInterestedIn(ReplicatedObjectId);
 					
 					const bool bIsSameAsBefore = ClientData.ObjectsWithSyncControl.Contains(ObjectId) == bIsValidForSending;
diff --git a/Engine/Plugins/Developer/Concert/ConcertSync/ConcertSyncServer/Source/ConcertSyncServer/Private/Replication/SyncControlManager.h b/Engine/Plugins/Developer/Concert/ConcertSync/ConcertSyncServer/Source/ConcertSyncServer/Private/Replication/SyncControlManager.h
--- a/Engine/Plugins/Developer/Concert/ConcertSync/ConcertSyncServer/Source/ConcertSyncServer/Private/Replication/SyncControlManager.h
+++ b/Engine/Plugins/Developer/Concert/ConcertSync/ConcertSyncServer/Source/ConcertSyncServer/Private/Replication/SyncControlManager.h
@@ -2,13 +2,17 @@
 
 #pragma once
 
-#include "ConcertMessageData.h"
 #include "Replication/Data/ObjectIds.h"
 
 #include "HAL/Platform.h"
 #include "Replication/Messages/SyncControl.h"
 #include "Templates/UnrealTemplate.h"
 
+namespace UE::ConcertSyncServer::Replication
+{
+	class FMuteManager;
+}
+
 class IConcertServerSession;
 enum class EConcertClientStatus : uint8;
 struct FConcertSessionClientInfo;
@@ -32,6 +36,7 @@ namespace UE::ConcertSyncServer::Replication
 		FSyncControlManager(
 			IConcertServerSession& Session UE_LIFETIMEBOUND,
 			FAuthorityManager& AuthorityManager UE_LIFETIMEBOUND,
+			FMuteManager& MuteManager UE_LIFETIMEBOUND,
 			const IRegistrationEnumerator& Getters UE_LIFETIMEBOUND
 			);
 		~FSyncControlManager();
@@ -43,7 +48,7 @@ namespace UE::ConcertSyncServer::Replication
 		FConcertReplication_ChangeSyncControl OnGenerateSyncControlForClientJoin(const FGuid& ClientId);
 		
 		/** Called by FConcertServerReplicationManager when client completes leaves replication. */
-		void OnClientLeft(const FGuid& ClientId) { HandleClientLeave(ClientId); }
+		void OnPostClientLeft(const FGuid& ClientId) { HandleClientLeave(ClientId); }
 
 	private:
 
@@ -51,6 +56,8 @@ namespace UE::ConcertSyncServer::Replication
 		IConcertServerSession& Session;
 		/** Used to detect whether a client has authority. */
 		FAuthorityManager& AuthorityManager;
+		/** Used to detect whether objects are globally muted. */
+		FMuteManager& MuteManager;
 
 		/** Callbacks for retrieving more info about client replication registration. */
 		const IRegistrationEnumerator& Getters;
@@ -66,6 +73,23 @@ namespace UE::ConcertSyncServer::Replication
 		/** Generates new sync control for client that is explicitly changing their authority. */
 		FConcertReplication_ChangeSyncControl OnGenerateSyncControlForAuthorityResponse(const FGuid& ClientId);
 
+		/**
+		 * Called when a client has caused mute state to change (e.g. due to removing object from stream, or explicitly muting it).
+		 * Called via delegate by FMuteManager.
+		 */
+		void OnUpdateSyncControlForIndirectMuteChange(const FGuid& ClientId)
+		{
+			// Do not send any update to this client (because they can infer the change themselves) but do update the sync control state.
+			RefreshClientSyncControl(ClientId);
+			// The other clients need to be notified in case an object was added which is now implicitly muted.
+			RefreshAndSendToAllClientsExcept(ClientId);
+		}
+		/** 
+		 * Updates sync control for all clients, sends an update to all clients but ClientId, and returns the sync control to embed into the mute response.
+		 * Called via delegate by FMuteManager.
+		 */
+		FConcertReplication_ChangeSyncControl OnGenerateSyncControlForClientMuteChange(const FGuid& ClientId);
+		
 		/** Cleans up the associated client data and updates sync control for all other clients. */
 		void HandleClientLeave(const FGuid& LeftClientId);
 
diff --git a/Engine/Plugins/Developer/Concert/ConcertSync/ConcertSyncTest/Source/ConcertSyncTest/Private/Replication/Components/ReplicatedObjectHierarchyCache.spec.cpp b/Engine/Plugins/Developer/Concert/ConcertSync/ConcertSyncTest/Source/ConcertSyncTest/Private/Replication/Components/ReplicatedObjectHierarchyCache.spec.cpp
--- a/Engine/Plugins/Developer/Concert/ConcertSync/ConcertSyncTest/Source/ConcertSyncTest/Private/Replication/Components/ReplicatedObjectHierarchyCache.spec.cpp
+++ b/Engine/Plugins/Developer/Concert/ConcertSync/ConcertSyncTest/Source/ConcertSyncTest/Private/Replication/Components/ReplicatedObjectHierarchyCache.spec.cpp
@@ -86,7 +86,7 @@ namespace UE::ConcertSyncTests::Replication
 				JoinRequest.Streams.Add(JoinStream);
 				
 				Cache.OnJoin(ClientId_One, JoinRequest);
-				Cache.OnPreClientLeft(ClientId_One, { JoinStream });
+				Cache.OnPostClientLeft(ClientId_One, { JoinStream });
 				
 				TestTrue(TEXT("IsEmpty"), Cache.IsEmpty());
 			});
@@ -102,7 +102,7 @@ namespace UE::ConcertSyncTests::Replication
 				
 				Cache.OnJoin(ClientId_One, JoinRequest);
 				Cache.OnJoin(ClientId_Two, JoinRequest);
-				Cache.OnPreClientLeft(ClientId_One, { JoinStream });
+				Cache.OnPostClientLeft(ClientId_One, { JoinStream });
 				
 				TestObjectInCache();
 			});
@@ -115,8 +115,8 @@ namespace UE::ConcertSyncTests::Replication
 				
 				Cache.OnJoin(ClientId_One, JoinRequest);
 				Cache.OnJoin(ClientId_Two, JoinRequest);
-				Cache.OnPreClientLeft(ClientId_One, { JoinStream });
-				Cache.OnPreClientLeft(ClientId_Two, { JoinStream });
+				Cache.OnPostClientLeft(ClientId_One, { JoinStream });
+				Cache.OnPostClientLeft(ClientId_Two, { JoinStream });
 				
 				TestTrue(TEXT("IsEmpty"), Cache.IsEmpty());
 			});
diff --git a/Engine/Plugins/Developer/Concert/ConcertSync/ConcertSyncTest/Source/ConcertSyncTest/Private/Replication/Components/SyncControlState.spec.cpp b/Engine/Plugins/Developer/Concert/ConcertSync/ConcertSyncTest/Source/ConcertSyncTest/Private/Replication/Components/SyncControlState.spec.cpp
--- a/Engine/Plugins/Developer/Concert/ConcertSync/ConcertSyncTest/Source/ConcertSyncTest/Private/Replication/Components/SyncControlState.spec.cpp
+++ b/Engine/Plugins/Developer/Concert/ConcertSync/ConcertSyncTest/Source/ConcertSyncTest/Private/Replication/Components/SyncControlState.spec.cpp
@@ -12,19 +12,24 @@ namespace UE::ConcertSyncTests::Replication::UI
 	FGuid SenderStreamId = FGuid::NewGuid();
 	/** Leverage to create a real FSoftObjectPath. */
 	TSharedPtr<FObjectTestReplicator> ObjectReplicator;
+	TSharedPtr<FObjectTestReplicator> SubobjectReplicator;
 	END_DEFINE_SPEC(FSyncControlStateSpec);
 
 	/** This tests that FSyncControlState correctly analyses requests and responses for aggregation. */
 	void FSyncControlStateSpec::Define()
 	{
+		using namespace ConcertSyncCore::Replication;
+		
 		BeforeEach([this]
 		{
 			ObjectReplicator = MakeShared<FObjectTestReplicator>();
+			SubobjectReplicator = ObjectReplicator->CreateSubobjectReplicator();
 		});
 		AfterEach([this]
 		{
 			// Test would hold onto this for rest of engine lifetime. Clean up this mini would-be leak.
 			ObjectReplicator.Reset();
+			SubobjectReplicator.Reset();
 		});
 
 		// At this point, everything should work server-side. Now test client-side prediction.
@@ -32,10 +37,10 @@ namespace UE::ConcertSyncTests::Replication::UI
 		{
 			It("Is correct with releasing authority with FConcertReplication_ChangeAuthority_Request", [this]()
 			{
-				ConcertSyncCore::Replication::FSyncControlState SyncControl = TSet<FConcertObjectInStreamID>{{ SenderStreamId, ObjectReplicator->TestObject }};
+				FSyncControlState SyncControl = TSet<FConcertObjectInStreamID>{{ SenderStreamId, ObjectReplicator->TestObject }};
+				
 				const FConcertReplication_ChangeAuthority_Request Request { .ReleaseAuthority = { { ObjectReplicator->TestObject, {{ SenderStreamId }} } } };
 				const FConcertReplication_ChangeAuthority_Response Response;
-				
 				TArray<FConcertObjectInStreamID> RemovedObjects;
 				SyncControl.AppendChanges(
 					Request,
@@ -50,11 +55,11 @@ namespace UE::ConcertSyncTests::Replication::UI
 			
 			It("Is correct with taking authority with FConcertReplication_ChangeAuthority_Request", [this]()
 			{
-				ConcertSyncCore::Replication::FSyncControlState SyncControl;
+				FSyncControlState SyncControl;
+				
 				const FConcertObjectInStreamID ObjectId{ SenderStreamId, ObjectReplicator->TestObject };
 				const FConcertReplication_ChangeAuthority_Request Request { .TakeAuthority = {{ ObjectReplicator->TestObject, {{ SenderStreamId }}}}};
 				const FConcertReplication_ChangeAuthority_Response Response { .SyncControl = {{{ ObjectId, true }}}};
-				
 				TArray<FConcertObjectInStreamID> AddedObjects;
 				SyncControl.AppendChanges(
 					Request,
@@ -69,16 +74,98 @@ namespace UE::ConcertSyncTests::Replication::UI
 
 			It("Is correct after releasing authority with FConcertReplication_ChangeStream_Request", [this]()
 			{
-				ConcertSyncCore::Replication::FSyncControlState SyncControl = TSet<FConcertObjectInStreamID>{{ SenderStreamId, ObjectReplicator->TestObject }};
+				FSyncControlState SyncControl = TSet<FConcertObjectInStreamID>{{ SenderStreamId, ObjectReplicator->TestObject }};
+				
 				FConcertReplication_ChangeStream_Request StreamChange;
 				StreamChange.ObjectsToRemove.Add({ SenderStreamId, ObjectReplicator->TestObject });
-
 				TArray<FConcertObjectInStreamID> RemovedObjects;
 				SyncControl.AppendChanges(StreamChange, [&RemovedObjects](const FConcertObjectInStreamID& Object){ RemovedObjects.Add(Object); });
 				
 				TestTrue(TEXT("Removed TestObject"), RemovedObjects.Contains(FConcertObjectInStreamID{ SenderStreamId, ObjectReplicator->TestObject }));
 				TestEqual(TEXT("Removed exactly 1 object"), RemovedObjects.Num(), 1);
 			});
+
+			Describe("With FConcertReplication_ChangeMuteState_Request", [this]()
+			{
+				It("Predicts for request muting single object", [this]()
+				{
+					FSyncControlState SyncControl = TSet<FConcertObjectInStreamID>{{ SenderStreamId, ObjectReplicator->TestObject }};
+					
+					FConcertReplication_ChangeMuteState_Request Request;
+					Request.ObjectsToMute.Add(ObjectReplicator->TestObject, { EConcertReplicationMuteOption::OnlyObject });
+					const FSyncControlState::FPredictedObjectRemoval Removal = SyncControl.AppendChanges(Request);
+					
+					TestEqual(TEXT("SyncControl.Num() == 0"), SyncControl.Num(), 0);
+				});
+				
+				It("Applies succeeded response unmuting object", [this]()
+				{
+					FSyncControlState SyncControl;
+					
+					FConcertReplication_ChangeMuteState_Request Request;
+					Request.ObjectsToUnmute.Add(ObjectReplicator->TestObject, { EConcertReplicationMuteOption::OnlyObject });
+					const FSyncControlState::FPredictedObjectRemoval Removal = SyncControl.AppendChanges(Request);
+					TestEqual(TEXT("SyncControl.Num() == 0"), SyncControl.Num(), 0);
+
+					FConcertReplication_ChangeMuteState_Response Response { EConcertReplicationMuteErrorCode::Accepted };
+					Response.SyncControl.NewControlStates.Add({ SenderStreamId, ObjectReplicator->TestObject }, true);
+					SyncControl.AppendChanges(Removal, Response);
+					
+					TestEqual(TEXT("SyncControl.Num() == 1"), SyncControl.Num(), 1);
+					TestTrue(TEXT("Contains TestObject"), SyncControl.IsObjectAllowed({ SenderStreamId, ObjectReplicator->TestObject }));
+				});
+				
+				It("Reverts failed response muting single object", [this]()
+				{
+					FSyncControlState SyncControl = TSet<FConcertObjectInStreamID>{{ SenderStreamId, ObjectReplicator->TestObject }};
+					
+					FConcertReplication_ChangeMuteState_Request Request;
+					Request.ObjectsToMute.Add(ObjectReplicator->TestObject, { EConcertReplicationMuteOption::OnlyObject });
+					const FSyncControlState::FPredictedObjectRemoval Removal = SyncControl.AppendChanges(Request);
+
+					const FConcertReplication_ChangeMuteState_Response Response { EConcertReplicationMuteErrorCode::Timeout };
+					SyncControl.AppendChanges(Removal, Response);
+					
+					TestEqual(TEXT("SyncControl.Num() == 1"), SyncControl.Num(), 1);
+					TestTrue(TEXT("Contains TestObject"), SyncControl.IsObjectAllowed({ SenderStreamId, ObjectReplicator->TestObject }));
+				});
+				
+				It("Predicts for request muting object & subobject", [this]()
+				{
+					FSyncControlState SyncControl = TSet<FConcertObjectInStreamID>
+					{
+						{ SenderStreamId, ObjectReplicator->TestObject },
+						{ SenderStreamId, SubobjectReplicator->TestObject }
+					};
+					
+					FConcertReplication_ChangeMuteState_Request Request;
+					Request.ObjectsToMute.Add(ObjectReplicator->TestObject, { EConcertReplicationMuteOption::ObjectAndSubobjects });
+					const FSyncControlState::FPredictedObjectRemoval Removal = SyncControl.AppendChanges(Request);
+					
+					TestEqual(TEXT("SyncControl.Num() == 0"), SyncControl.Num(), 0);
+				});
+
+				It("Applies implicitly unmuted objects", [this]()
+				{
+					FSyncControlState SyncControl;
+
+					FConcertReplication_ChangeMuteState_Request Request;
+					Request.ObjectsToUnmute.Add(ObjectReplicator->TestObject, { EConcertReplicationMuteOption::ObjectAndSubobjects });
+					const FSyncControlState::FPredictedObjectRemoval Removal = SyncControl.AppendChanges(Request);
+					
+					FConcertReplication_ChangeMuteState_Response Response { EConcertReplicationMuteErrorCode::Accepted };
+					Response.SyncControl.NewControlStates =
+					{
+						{ { SenderStreamId, ObjectReplicator->TestObject }, true },
+						{ { SenderStreamId, SubobjectReplicator->TestObject }, true },
+					};
+					SyncControl.AppendChanges(Removal, Response);
+					
+					TestEqual(TEXT("SyncControl.Num() == 2"), SyncControl.Num(), 2);
+					TestTrue(TEXT("Contains root"), SyncControl.IsObjectAllowed({ SenderStreamId, ObjectReplicator->TestObject }));
+					TestTrue(TEXT("Contains subobject"), SyncControl.IsObjectAllowed({ SenderStreamId, SubobjectReplicator->TestObject }));
+				});
+			});
 		});
 	}
 }
diff --git a/Engine/Plugins/Developer/Concert/ConcertSync/ConcertSyncTest/Source/ConcertSyncTest/Private/Replication/Muting/MuteAndQuery.spec.cpp b/Engine/Plugins/Developer/Concert/ConcertSync/ConcertSyncTest/Source/ConcertSyncTest/Private/Replication/Muting/MuteAndQuery.spec.cpp
new file mode 100644
--- /dev/null
+++ b/Engine/Plugins/Developer/Concert/ConcertSync/ConcertSyncTest/Source/ConcertSyncTest/Private/Replication/Muting/MuteAndQuery.spec.cpp
@@ -0,0 +1,420 @@
+﻿// Copyright Epic Games, Inc. All Rights Reserved.
+
+#include "Misc/AutomationTest.h"
+#include "Replication/Util/Spec/ObjectTestReplicator.h"
+#include "Replication/Util/Spec/ReplicationClient.h"
+#include "Replication/Util/Spec/ReplicationServer.h"
+
+namespace UE::ConcertSyncTests::Replication
+{
+	/** This tests that muting & unmuting retains a consistent check across various types of changes. */
+	BEGIN_DEFINE_SPEC(FMuteAndQuerySpec, "Editor.Concert.Replication.Muting", EAutomationTestFlags::EditorContext | EAutomationTestFlags::EngineFilter)
+	TSharedPtr<FReplicationServer> Server;
+	FReplicationClient* Sender = nullptr;
+
+	FGuid StreamId_Bar					= FGuid::NewGuid();
+	FGuid StreamId_BarSubobject			= FGuid::NewGuid();
+	FGuid StreamId_BarNestedSubobject	= FGuid::NewGuid();
+	// Leverage FObjectTestReplicator to create more UObjects
+	TSharedPtr<FObjectTestReplicator> Replicator_Bar;			
+	TSharedPtr<FObjectTestReplicator> Replicator_BarSubobject;		
+	TSharedPtr<FObjectTestReplicator> Replicator_BarNestedSubobject;
+	END_DEFINE_SPEC(FMuteAndQuerySpec);
+
+	/** This tests that muting requests work when EConcertSyncSessionFlags::ShouldAllowGlobalMuting is set. */
+	void FMuteAndQuerySpec::Define()
+	{
+		BeforeEach([this]
+		{
+			Server		= MakeShared<FReplicationServer>(*this);
+			Sender		= &Server->ConnectClient();
+
+			Replicator_Bar					= MakeShared<FObjectTestReplicator>();
+			Replicator_BarSubobject			= Replicator_Bar->CreateSubobjectReplicator();
+			Replicator_BarNestedSubobject	= Replicator_BarSubobject->CreateSubobjectReplicator();
+		});
+		AfterEach([this]
+		{
+			// Test would hold onto this for rest of engine lifetime. Clean up this mini would-be leak.
+			Server.Reset();
+			Replicator_Bar.Reset();
+			Replicator_BarSubobject.Reset();
+			Replicator_BarNestedSubobject.Reset();
+		});
+
+		// These tests involve only using mute & unmute requests (for modifying stream contents see further below)
+		Describe("Muting hierarchy (only mute & unmute requests)", [this]()
+		{
+			BeforeEach([this]()
+			{
+				ConcertSyncClient::Replication::FJoinReplicatedSessionArgs SenderJoinArgs;
+				SenderJoinArgs.Streams.Add(Replicator_Bar->CreateStream(StreamId_Bar));
+				SenderJoinArgs.Streams.Add(Replicator_BarSubobject->CreateStream(StreamId_BarSubobject));
+				SenderJoinArgs.Streams.Add(Replicator_BarNestedSubobject->CreateStream(StreamId_BarNestedSubobject));
+				
+				Sender->JoinReplication(SenderJoinArgs);
+				Sender->GetClientReplicationManager().MuteObjects({ Replicator_Bar->TestObject }, EConcertReplicationMuteOption::ObjectAndSubobjects);
+			});
+			
+			It("Query entire hierarchy", [this]()
+			{
+				// Requesting all muted objects
+				bool bReceivedQueryAllEvent = false;
+				Sender->GetClientReplicationManager()
+					.QueryMuteState()
+					.Next([this, &bReceivedQueryAllEvent](FConcertReplication_QueryMuteState_Response&& Response)
+					{
+						bReceivedQueryAllEvent = true;
+						TestEqual(TEXT("ExplicitlyMutedObjects.Num() == 1"), Response.ExplicitlyMutedObjects.Num(), 1);
+						TestEqual(TEXT("ExplicitlyUnmutedObjects.Num() == 0"), Response.ExplicitlyUnmutedObjects.Num(), 0);
+						TestEqual(TEXT("ImplicitlyMutedObjects.Num() == 2"), Response.ImplicitlyMutedObjects.Num(), 2);
+						TestEqual(TEXT("ImplicitlyUnmutedObjects.Num() == 0"), Response.ImplicitlyUnmutedObjects.Num(), 0);
+						
+						const FConcertReplication_ObjectMuteSetting* ObjectOneMuteSetting = Response.ExplicitlyMutedObjects.Find(Replicator_Bar->TestObject);
+						TestTrue(TEXT("ExplicitlyMutedObjects.Contains(Bar)"), ObjectOneMuteSetting && ObjectOneMuteSetting->Flags == EConcertReplicationMuteOption::ObjectAndSubobjects);
+						TestTrue(TEXT("ImplicitlyMutedObjects.Contains(BarSubobject)"), Response.ImplicitlyMutedObjects.Contains(Replicator_BarSubobject->TestObject));
+						TestTrue(TEXT("ImplicitlyMutedObjects.Contains(BarNestedSubobject)"), Response.ImplicitlyMutedObjects.Contains(Replicator_BarNestedSubobject->TestObject));
+						
+					});
+				TestTrue(TEXT("Received query all event"), bReceivedQueryAllEvent);
+			});
+			
+			It("Query only BarSubobject", [this]()
+			{
+				Sender->GetClientReplicationManager()
+					.QueryMuteState({ Replicator_BarSubobject->TestObject })
+					.Next([this](FConcertReplication_QueryMuteState_Response&& Response)
+					{
+						TestEqual(TEXT("ExplicitlyMutedObjects.Num() == 0"), Response.ExplicitlyMutedObjects.Num(), 0);
+						TestEqual(TEXT("ExplicitlyUnmutedObjects.Num() == 0"), Response.ExplicitlyUnmutedObjects.Num(), 0);
+						TestEqual(TEXT("ImplicitlyMutedObjects.Num() == 1"), Response.ImplicitlyMutedObjects.Num(), 1);
+						TestEqual(TEXT("ImplicitlyUnmutedObjects.Num() == 0"), Response.ImplicitlyUnmutedObjects.Num(), 0);
+						
+						TestTrue(TEXT("ImplicitlyMutedObjects.Contains(BarSubobject)"), Response.ImplicitlyMutedObjects.Contains(Replicator_BarSubobject->TestObject));
+					});
+			});
+			
+			It("Unmute subobject", [this]()
+			{
+				Sender->GetClientReplicationManager().UnmuteObjects({ Replicator_BarSubobject->TestObject }, EConcertReplicationMuteOption::OnlyObject);
+				
+				bool bReceivedEvent = false;
+				Sender->GetClientReplicationManager()
+					.QueryMuteState({ Replicator_Bar->TestObject, Replicator_BarSubobject->TestObject })
+					.Next([this, &bReceivedEvent](FConcertReplication_QueryMuteState_Response&& Response)
+					{
+						bReceivedEvent = true;
+						
+						TestEqual(TEXT("ExplicitlyMutedObjects.Num() == 1"), Response.ExplicitlyMutedObjects.Num(), 1);
+						TestEqual(TEXT("ExplicitlyUnmutedObjects.Num() == 1"), Response.ExplicitlyUnmutedObjects.Num(), 1);
+						TestEqual(TEXT("ImplicitlyMutedObjects.Num() == 0"), Response.ImplicitlyMutedObjects.Num(), 0);
+						TestEqual(TEXT("ImplicitlyUnmutedObjects.Num() == 0"), Response.ImplicitlyUnmutedObjects.Num(), 0);
+					
+						const FConcertReplication_ObjectMuteSetting* ObjectOneMuteSetting = Response.ExplicitlyMutedObjects.Find(Replicator_Bar->TestObject);
+						TestTrue(TEXT("Root object still explicitly muted"), ObjectOneMuteSetting && ObjectOneMuteSetting->Flags == EConcertReplicationMuteOption::ObjectAndSubobjects);
+					
+						const FConcertReplication_ObjectMuteSetting* ObjectTwoMuteSetting = Response.ExplicitlyUnmutedObjects.Find(Replicator_BarSubobject->TestObject);
+						TestTrue(TEXT("Subobject explicitly unmuted"), ObjectTwoMuteSetting && ObjectTwoMuteSetting->Flags == EConcertReplicationMuteOption::OnlyObject);
+					});
+				TestTrue(TEXT("Received query response"), bReceivedEvent);
+			});
+
+			It("Unmuting hiearchy by unmuting root object ", [this]()
+			{
+				Sender->GetClientReplicationManager().UnmuteObjects({ Replicator_Bar->TestObject });
+				
+				Sender->GetClientReplicationManager()
+					.QueryMuteState()
+					.Next([this](FConcertReplication_QueryMuteState_Response&& Response)
+					{
+						TestTrue(TEXT("IsEmpty()"), Response.IsEmpty());
+					});
+			});
+			
+			It("Unmute hierarchy by changing root to EConcertReplicationMuteFlags::None", [this]()
+			{
+				Sender->GetClientReplicationManager().MuteObjects({ Replicator_Bar->TestObject }, EConcertReplicationMuteOption::OnlyObject);
+				
+				Sender->GetClientReplicationManager()
+					.QueryMuteState({ Replicator_BarSubobject->TestObject })
+					.Next([this](FConcertReplication_QueryMuteState_Response&& Response)
+					{
+						TestTrue(TEXT("IsEmpty"), Response.IsEmpty());
+					});
+			});
+
+			It("Double mute hierarchy", [this]()
+			{
+				// The point of this test is that in this case
+				// - Bar <- is IncludeSubobjects
+				//	- BarSubobject <- change from IncludeSubobjects to None
+				//		- BarNestedSubobject <- stays implicitly muted
+				// the system detects that NestedSubobject should remain muted due to Root.
+				Sender->GetClientReplicationManager().MuteObjects({ Replicator_BarSubobject->TestObject }, EConcertReplicationMuteOption::ObjectAndSubobjects);
+				Sender->GetClientReplicationManager().UnmuteObjects({ Replicator_BarSubobject->TestObject }, EConcertReplicationMuteOption::OnlyObject);
+				Sender->GetClientReplicationManager()
+					.QueryMuteState()
+					.Next([this](FConcertReplication_QueryMuteState_Response&& Response)
+					{
+						TestEqual(TEXT("ExplicitlyMutedObjects.Num() == 1"), Response.ExplicitlyMutedObjects.Num(), 1);
+						TestEqual(TEXT("ExplicitlyUnmutedObjects.Num() == 1"), Response.ExplicitlyUnmutedObjects.Num(), 1);
+						TestEqual(TEXT("ImplicitlyMutedObjects.Num() == 1"), Response.ImplicitlyMutedObjects.Num(), 1);
+						TestEqual(TEXT("ImplicitlyUnmutedObjects.Num() == 0"), Response.ImplicitlyUnmutedObjects.Num(), 0);
+					
+						const FConcertReplication_ObjectMuteSetting* BarMuteSetting = Response.ExplicitlyMutedObjects.Find(Replicator_Bar->TestObject);
+						TestTrue(TEXT("ExplicitlyMutedObjects.Contains(Bar)"), BarMuteSetting && BarMuteSetting->Flags == EConcertReplicationMuteOption::ObjectAndSubobjects);
+						const FConcertReplication_ObjectMuteSetting* SubobjectMuteSetting = Response.ExplicitlyUnmutedObjects.Find(Replicator_BarSubobject->TestObject);
+						TestTrue(TEXT("ExplicitlyUnmutedObjects.Contains(BarSubobject)"), SubobjectMuteSetting && SubobjectMuteSetting->Flags == EConcertReplicationMuteOption::OnlyObject);
+						TestTrue(TEXT("ImplicitlyMutedObjects.Contains(BarNestedSubobject)"), Response.ImplicitlyMutedObjects.Contains(Replicator_BarNestedSubobject->TestObject));
+					});
+			});
+		});
+
+		It("Changes are atomic", [this]()
+		{
+			Sender->JoinReplication(Replicator_Bar->CreateSenderArgs(StreamId_Bar));
+			
+			// This is supposed to fail because Replicator_BarSubobject is not registered
+			Sender->GetClientReplicationManager().MuteObjects({ Replicator_Bar->TestObject, Replicator_BarSubobject->TestObject });
+
+			Sender->GetClientReplicationManager()
+				.QueryMuteState()
+				.Next([this](FConcertReplication_QueryMuteState_Response&& Response)
+				{
+					TestTrue(TEXT("IsEmpty()"), Response.IsEmpty());
+				});
+		});
+
+		It("Implicitly unmute nested subobject", [this]()
+		{
+			ConcertSyncClient::Replication::FJoinReplicatedSessionArgs SenderJoinArgs;
+			SenderJoinArgs.Streams.Add(Replicator_Bar->CreateStream(StreamId_Bar));
+			SenderJoinArgs.Streams.Add(Replicator_BarSubobject->CreateStream(StreamId_BarSubobject));
+			SenderJoinArgs.Streams.Add(Replicator_BarNestedSubobject->CreateStream(StreamId_BarNestedSubobject));
+			Sender->JoinReplication(SenderJoinArgs);
+			
+			FConcertReplication_ChangeMuteState_Request MuteRequest;
+			MuteRequest.ObjectsToMute.Add(Replicator_Bar->TestObject, { EConcertReplicationMuteOption::ObjectAndSubobjects });
+			MuteRequest.ObjectsToUnmute.Add(Replicator_BarSubobject->TestObject, { EConcertReplicationMuteOption::ObjectAndSubobjects });
+			Sender->GetClientReplicationManager().ChangeMuteState(MuteRequest);
+			
+			Sender->GetClientReplicationManager()
+				.QueryMuteState()
+				.Next([this](FConcertReplication_QueryMuteState_Response&& Response)
+				{
+					TestTrue(TEXT("IsSuccess"), Response.IsSuccess());
+					
+					TestEqual(TEXT("ExplicitlyMutedObjects.Num() == 1"), Response.ExplicitlyMutedObjects.Num(), 1);
+					TestEqual(TEXT("ExplicitlyUnmutedObjects.Num() == 1"), Response.ExplicitlyUnmutedObjects.Num(), 1);
+					TestEqual(TEXT("ImplicitlyMutedObjects.Num() == 1"), Response.ImplicitlyMutedObjects.Num(), 0);
+					TestEqual(TEXT("ImplicitlyUnmutedObjects.Num() == 1"), Response.ImplicitlyUnmutedObjects.Num(), 1);
+
+					const FConcertReplication_ObjectMuteSetting* ExplicitlyMuted = Response.ExplicitlyMutedObjects.Find(Replicator_Bar->TestObject);
+					const FConcertReplication_ObjectMuteSetting* ExplicitlyUnmuted = Response.ExplicitlyUnmutedObjects.Find(Replicator_BarSubobject->TestObject);
+					
+					TestTrue(TEXT("Root is muted"), ExplicitlyMuted != nullptr);
+					TestTrue(TEXT("Root has right flags"), ExplicitlyMuted && ExplicitlyMuted->Flags == EConcertReplicationMuteOption::ObjectAndSubobjects);
+					
+					TestTrue(TEXT("Subobject is unmuted"), ExplicitlyUnmuted != nullptr);
+					TestTrue(TEXT("Subobject has right flags"), ExplicitlyUnmuted && ExplicitlyUnmuted->Flags == EConcertReplicationMuteOption::ObjectAndSubobjects);
+					
+					TestTrue(TEXT("NestedSubobject is unmuted"), Response.ImplicitlyUnmutedObjects.Contains(Replicator_BarNestedSubobject->TestObject));
+				});
+		});
+
+		// Tests that involve FConcertReplication_ChangeStream_Request while using mutes
+		Describe("Muting hierarchy (include changing stream structure)", [this]()
+		{
+			Describe("Keep hierarchy muted when root object is removed from one stream but referenced by another stream", [this]()
+			{
+				BeforeEach([this]()
+				{
+					ConcertSyncClient::Replication::FJoinReplicatedSessionArgs SenderJoinArgs;
+					SenderJoinArgs.Streams.Add(Replicator_Bar->CreateStream(StreamId_Bar));
+					SenderJoinArgs.Streams.Add(Replicator_BarSubobject->CreateStream(StreamId_BarSubobject));
+					
+					Sender->JoinReplication(SenderJoinArgs);
+					Sender->GetClientReplicationManager().MuteObjects({ Replicator_Bar->TestObject }, EConcertReplicationMuteOption::ObjectAndSubobjects);
+				});
+
+				AfterEach([this]()
+				{
+					Sender->GetClientReplicationManager()
+						.QueryMuteState()
+						.Next([this](FConcertReplication_QueryMuteState_Response&& Response)
+						{
+							TestEqual(TEXT("ExplicitlyMutedObjects.Num() == 1"), Response.ExplicitlyMutedObjects.Num(), 1);
+							TestEqual(TEXT("ExplicitlyUnmutedObjects.Num() == 0"), Response.ExplicitlyUnmutedObjects.Num(), 0);
+							TestEqual(TEXT("ImplicitlyMutedObjects.Num() == 1"), Response.ImplicitlyMutedObjects.Num(), 1);
+							TestEqual(TEXT("ImplicitlyUnmutedObjects.Num() == 0"), Response.ImplicitlyUnmutedObjects.Num(), 0);
+							
+							TestTrue(TEXT("ExplicitlyMutedObjects.Contains(Bar)"), Response.ExplicitlyMutedObjects.Contains(Replicator_Bar->TestObject));
+							TestTrue(TEXT("ExplicitlyMutedObjects.Contains(Subobject)"), Response.ImplicitlyMutedObjects.Contains(Replicator_BarSubobject->TestObject));
+						});
+				});
+				
+				It("Single FConcertReplication_ChangeStream_Request for adding & removing root object", [this]()
+				{
+					FConcertReplication_ChangeStream_Request Request;
+					Request.StreamsToRemove.Add(StreamId_Bar);
+					Request.StreamsToAdd.Add(Replicator_Bar->CreateStream(FGuid::NewGuid()));
+					Sender->GetClientReplicationManager().ChangeStream(Request);
+				});
+				
+				It("Separate FConcertReplication_ChangeStream_Requests for adding & removing root object", [this]()
+				{
+					FConcertReplication_ChangeStream_Request InitialRequest;
+					InitialRequest.StreamsToAdd.Add(Replicator_Bar->CreateStream(FGuid::NewGuid()));
+					Sender->GetClientReplicationManager().ChangeStream(InitialRequest);
+					
+					FConcertReplication_ChangeStream_Request Request;
+					Request.StreamsToRemove.Add(StreamId_Bar);
+					Sender->GetClientReplicationManager().ChangeStream(Request);
+				});
+			});
+			
+			Describe("Removing multiple parents retains and cleans up state correctly.", [this]()
+			{
+				BeforeEach([this]()
+				{
+					ConcertSyncClient::Replication::FJoinReplicatedSessionArgs SenderJoinArgs;
+					SenderJoinArgs.Streams.Add(Replicator_Bar->CreateStream(StreamId_Bar));
+					SenderJoinArgs.Streams.Add(Replicator_BarSubobject->CreateStream(StreamId_BarSubobject));
+					SenderJoinArgs.Streams.Add(Replicator_BarNestedSubobject->CreateStream(StreamId_BarNestedSubobject));
+				
+					Sender->JoinReplication(SenderJoinArgs);
+					Sender->GetClientReplicationManager().MuteObjects({ Replicator_Bar->TestObject }, EConcertReplicationMuteOption::ObjectAndSubobjects);
+					Sender->GetClientReplicationManager().MuteObjects({ Replicator_BarSubobject->TestObject }, EConcertReplicationMuteOption::ObjectAndSubobjects);
+				});
+
+				Describe("Retain all mute state", [this]()
+				{
+					AfterEach([this]()
+					{
+						Sender->GetClientReplicationManager()
+							.QueryMuteState()
+							.Next([this](FConcertReplication_QueryMuteState_Response&& Response)
+							{
+								TestEqual(TEXT("ExplicitlyMutedObjects.Num()"), Response.ExplicitlyMutedObjects.Num(), 2);
+								TestEqual(TEXT("ExplicitlyUnmutedObjects.Num()"), Response.ExplicitlyUnmutedObjects.Num(), 0);
+								TestEqual(TEXT("ImplicitlyMutedObjects.Num()"), Response.ImplicitlyMutedObjects.Num(), 1);
+								TestEqual(TEXT("ImplicitlyUnmutedObjects.Num()"), Response.ImplicitlyUnmutedObjects.Num(), 0);
+							
+								TestTrue(TEXT("ImplicitlyMutedObjects.Contains(Bar)"), Response.ExplicitlyMutedObjects.Contains(Replicator_Bar->TestObject));
+								TestTrue(TEXT("ImplicitlyMutedObjects.Contains(BarSubobject)"), Response.ExplicitlyMutedObjects.Contains(Replicator_BarSubobject->TestObject));
+								TestTrue(TEXT("ImplicitlyMutedObjects.Contains(BarNestedSubobject)"), Response.ImplicitlyMutedObjects.Contains(Replicator_BarNestedSubobject->TestObject));
+							});
+					});
+					
+					It("After removing middle parent", [this]()
+					{
+						const FConcertReplication_ChangeStream_Request Request { .StreamsToRemove = { StreamId_BarSubobject } };
+						Sender->GetClientReplicationManager().ChangeStream(Request);
+					});
+					It("After removing root", [this]()
+					{
+						const FConcertReplication_ChangeStream_Request Request { .StreamsToRemove = { StreamId_Bar } };
+						Sender->GetClientReplicationManager().ChangeStream(Request);
+					});
+					It("Remove root and parent in single request", [this]()
+					{
+						const FConcertReplication_ChangeStream_Request Request { .StreamsToRemove = { StreamId_Bar, StreamId_BarSubobject } };
+						Sender->GetClientReplicationManager().ChangeStream(Request);
+					});
+					It("Remove root, then parent", [this]()
+					{
+						const FConcertReplication_ChangeStream_Request RemoveBar { .StreamsToRemove = { StreamId_Bar } };
+						Sender->GetClientReplicationManager().ChangeStream(RemoveBar);
+
+						const FConcertReplication_ChangeStream_Request RemoveBarSubobject { .StreamsToRemove = { StreamId_BarSubobject } };
+						Sender->GetClientReplicationManager().ChangeStream(RemoveBarSubobject);
+					});
+					It ("Remove parent, then root", [this]()
+					{
+						const FConcertReplication_ChangeStream_Request RemoveBarSubobject { .StreamsToRemove = { StreamId_BarSubobject } };
+						Sender->GetClientReplicationManager().ChangeStream(RemoveBarSubobject);
+						
+						const FConcertReplication_ChangeStream_Request RemoveBar { .StreamsToRemove = { StreamId_Bar } };
+						Sender->GetClientReplicationManager().ChangeStream(RemoveBar);
+					});
+				});
+				
+				Describe("After unmuting root and parent, nothing retains mute state", [this]()
+				{
+					AfterEach([this]()
+					{
+						Sender->GetClientReplicationManager()
+							.QueryMuteState()
+							.Next([this](FConcertReplication_QueryMuteState_Response&& Response)
+							{
+								TestEqual(TEXT("ExplicitlyMutedObjects.Num() == 0"), Response.ExplicitlyMutedObjects.Num(), 0);
+								TestEqual(TEXT("ExplicitlyUnmutedObjects.Num() == 0"), Response.ExplicitlyUnmutedObjects.Num(), 0);
+								TestEqual(TEXT("ImplicitlyMutedObjects.Num() == 0"), Response.ImplicitlyMutedObjects.Num(), 0);
+								TestEqual(TEXT("ImplicitlyUnmutedObjects.Num() == 0"), Response.ImplicitlyUnmutedObjects.Num(), 0);
+							});
+					});
+
+					It("Unmute root and parent in single request", [this]()
+					{
+						Sender->GetClientReplicationManager().UnmuteObjects({ Replicator_Bar->TestObject, Replicator_BarSubobject->TestObject });
+					});
+					It("Unmute root, then parent", [this]()
+					{
+						Sender->GetClientReplicationManager().UnmuteObjects({ Replicator_Bar->TestObject });
+						Sender->GetClientReplicationManager().UnmuteObjects({ Replicator_BarSubobject->TestObject });
+					});
+					It ("Unmute parent, then root", [this]()
+					{
+						Sender->GetClientReplicationManager().UnmuteObjects({ Replicator_BarSubobject->TestObject });
+						Sender->GetClientReplicationManager().UnmuteObjects({ Replicator_Bar->TestObject });
+					});
+				});
+
+				Describe("After removing root and parent, nothing retains mute state", [this]()
+				{
+					BeforeEach([this]()
+					{
+						// The nested subobject will keep alive the entire mute state if not removed because it causes its outers to be referenced.
+						const FConcertReplication_ChangeStream_Request Request { .StreamsToRemove = { StreamId_BarNestedSubobject } };
+						Sender->GetClientReplicationManager().ChangeStream(Request);
+					});
+					
+					AfterEach([this]()
+					{
+						Sender->GetClientReplicationManager()
+							.QueryMuteState()
+							.Next([this](FConcertReplication_QueryMuteState_Response&& Response)
+							{
+								TestEqual(TEXT("ExplicitlyMutedObjects.Num() == 0"), Response.ExplicitlyMutedObjects.Num(), 0);
+								TestEqual(TEXT("ExplicitlyUnmutedObjects.Num() == 0"), Response.ExplicitlyUnmutedObjects.Num(), 0);
+								TestEqual(TEXT("ImplicitlyMutedObjects.Num() == 0"), Response.ImplicitlyMutedObjects.Num(), 0);
+								TestEqual(TEXT("ImplicitlyUnmutedObjects.Num() == 0"), Response.ImplicitlyUnmutedObjects.Num(), 0);
+							});
+					});
+					
+					It("Remove root and parent in single request", [this]()
+					{
+						const FConcertReplication_ChangeStream_Request Request { .StreamsToRemove = { StreamId_Bar, StreamId_BarSubobject } };
+						Sender->GetClientReplicationManager().ChangeStream(Request);
+					});
+					It("Remove root, then parent", [this]()
+					{
+						const FConcertReplication_ChangeStream_Request RemoveBar { .StreamsToRemove = { StreamId_Bar } };
+						Sender->GetClientReplicationManager().ChangeStream(RemoveBar);
+
+						const FConcertReplication_ChangeStream_Request RemoveBarSubobject { .StreamsToRemove = { StreamId_BarSubobject } };
+						Sender->GetClientReplicationManager().ChangeStream(RemoveBarSubobject);
+					});
+					It ("Remove parent, then root", [this]()
+					{
+						const FConcertReplication_ChangeStream_Request RemoveBarSubobject { .StreamsToRemove = { StreamId_BarSubobject } };
+						Sender->GetClientReplicationManager().ChangeStream(RemoveBarSubobject);
+						
+						const FConcertReplication_ChangeStream_Request RemoveBar { .StreamsToRemove = { StreamId_Bar } };
+						Sender->GetClientReplicationManager().ChangeStream(RemoveBar);
+					});
+				});
+			});
+		});
+	}
+}
\ No newline at end of file
diff --git a/Engine/Plugins/Developer/Concert/ConcertSync/ConcertSyncTest/Source/ConcertSyncTest/Private/Replication/Muting/MuteAndSyncControl.spec.cpp b/Engine/Plugins/Developer/Concert/ConcertSync/ConcertSyncTest/Source/ConcertSyncTest/Private/Replication/Muting/MuteAndSyncControl.spec.cpp
new file mode 100644
--- /dev/null
+++ b/Engine/Plugins/Developer/Concert/ConcertSync/ConcertSyncTest/Source/ConcertSyncTest/Private/Replication/Muting/MuteAndSyncControl.spec.cpp
@@ -0,0 +1,153 @@
+﻿// Copyright Epic Games, Inc. All Rights Reserved.
+
+#include "Misc/AutomationTest.h"
+#include "Replication/Util/ConcertClientReplicationBridgeMock.h"
+#include "Replication/Util/Spec/ObjectTestReplicator.h"
+#include "Replication/Util/Spec/ReplicationClient.h"
+#include "Replication/Util/Spec/ReplicationServer.h"
+
+namespace UE::ConcertSyncTests::Replication
+{
+	BEGIN_DEFINE_SPEC(FMuteAndSyncControlSpec, "Editor.Concert.Replication.Muting", EAutomationTestFlags::EditorContext | EAutomationTestFlags::EngineFilter)
+	TSharedPtr<FReplicationServer> Server;
+	FReplicationClient* Sender = nullptr;
+	FReplicationClient* Receiver = nullptr;
+
+	FGuid ObjectStreamId = FGuid::NewGuid();
+	TSharedPtr<FObjectTestReplicator> ObjectReplicator;
+
+	FObjectReplicationContext MakeSenderToReceiverContext() const { return { *Sender, *Server, *Receiver }; }
+	END_DEFINE_SPEC(FMuteAndSyncControlSpec);
+
+	/** This tests that muting and sync control interact correctly. */
+	void FMuteAndSyncControlSpec::Define()
+	{
+		BeforeEach([this]
+		{
+			Server = MakeShared<FReplicationServer>(*this);
+			Sender = &Server->ConnectClient();
+			Receiver = &Server->ConnectClient();
+
+			ObjectReplicator = MakeShared<FObjectTestReplicator>();
+		});
+		AfterEach([this]
+		{
+			// Test would hold onto this for rest of engine lifetime. Clean up this mini would-be leak.
+			Server.Reset();
+			ObjectReplicator.Reset();
+			
+			Sender = nullptr;
+			Receiver = nullptr;
+		});
+		
+		Describe("When client mutes / unmutes its own object", [this]()
+		{
+			BeforeEach([this]()
+			{
+				Sender->JoinReplication(ObjectReplicator->CreateSenderArgs(ObjectStreamId));
+				// Tells client that the object is ready to replicate
+				Sender->GetBridgeMock().InjectAvailableObject(*ObjectReplicator->TestObject);
+				// Tells server intent to replicate the object
+				Sender->GetClientReplicationManager().TakeAuthorityOver({ ObjectReplicator->TestObject });
+
+				// This gives sync control meaning that now the object can be replicated.
+				Receiver->JoinReplicationAsListener({ ObjectReplicator->TestObject });
+
+				Sender->GetClientSessionMock()->RegisterCustomEventHandler<FConcertReplication_ChangeSyncControl>(
+					[this](const FConcertSessionContext&, const FConcertReplication_ChangeSyncControl& Event)
+					{
+						AddError(TEXT("Requesting client was not supposed to receive FConcertReplication_ChangeSyncControl"));
+					});
+				Sender->GetClientReplicationManager().MuteObjects({ ObjectReplicator->TestObject });
+			});
+			
+			It("Client predicts new sync control for muted object", [this]()
+			{
+				TestTrue(TEXT("Has no sync control"), !Sender->GetClientReplicationManager().HasSyncControl({ ObjectStreamId, ObjectReplicator->TestObject }));
+				
+				ObjectReplicator->SimulateSendObjectToReceiver(*this, MakeSenderToReceiverContext(), { ObjectStreamId },
+					[this](const FConcertSessionContext& Context, const FConcertReplication_BatchReplicationEvent& Event)
+					{
+						// If the client predicted the sync control correctly, it should not attempt to send to the server.
+						AddError(TEXT("Sending client did not predict sync control correctly"));
+					});
+				ObjectReplicator->TestValuesWereNotReplicated(*this);
+			});
+
+			It("Client receives new sync control for unmuted object", [this]()
+			{
+				Sender->GetClientReplicationManager()
+					.UnmuteObjects({ ObjectReplicator->TestObject })
+					.Next([this](FConcertReplication_ChangeMuteState_Response&& Response)
+					{
+						TestTrue(TEXT("Response.IsSuccess()"), Response.IsSuccess());
+						TestEqual(TEXT("Response.SyncControl.NewControlStates.Num() == 1"), Response.SyncControl.NewControlStates.Num(), 1);
+						
+						const bool* bIsEnabled = Response.SyncControl.NewControlStates.Find({ ObjectStreamId, ObjectReplicator->TestObject });
+						TestTrue(TEXT("Object has regained sync control"), bIsEnabled && *bIsEnabled);
+					});
+				
+				TestTrue(TEXT("Has sync control"), Sender->GetClientReplicationManager().HasSyncControl({ ObjectStreamId, ObjectReplicator->TestObject }));
+				ObjectReplicator->SimulateSendObjectToReceiver(*this, MakeSenderToReceiverContext(), { ObjectStreamId });
+				ObjectReplicator->TestValuesWereReplicated(*this);
+			});
+		});
+
+		Describe("Clients receive FConcertReplication_ChangeSyncControl when other clients", [this]()
+		{
+			BeforeEach([this]()
+			{
+				Sender->JoinReplication(ObjectReplicator->CreateSenderArgs(ObjectStreamId));
+				// Tells client that the object is ready to replicate
+				Sender->GetBridgeMock().InjectAvailableObject(*ObjectReplicator->TestObject);
+				// Tells server intent to replicate the object
+				Sender->GetClientReplicationManager().TakeAuthorityOver({ ObjectReplicator->TestObject });
+
+				Receiver->JoinReplicationAsListener({});
+			});
+			
+			It("Mute the object", [this]()
+			{
+				bool bReceivedEvent = false;
+				Sender->GetClientSessionMock()->RegisterCustomEventHandler<FConcertReplication_ChangeSyncControl>(
+					[this, &bReceivedEvent](const FConcertSessionContext&, const FConcertReplication_ChangeSyncControl& Event)
+					{
+						bReceivedEvent = true;
+						TestEqual(TEXT("Event.NewControlStates.Num() == 1"), Event.NewControlStates.Num(), 1);
+
+						const bool* bState = Event.NewControlStates.Find({ ObjectStreamId, ObjectReplicator->TestObject });
+						TestTrue(TEXT("Lost sync control"), bState && !*bState);
+					});
+				
+				Receiver->GetClientReplicationManager().MuteObjects({ ObjectReplicator->TestObject });
+				TestTrue(TEXT("Received event"), bReceivedEvent);
+			});
+
+			It("Unmute the object", [this]()
+			{
+				Receiver->GetClientReplicationManager().MuteObjects({ ObjectReplicator->TestObject });
+				
+				bool bReceivedEvent = false;
+				Sender->GetClientSessionMock()->RegisterCustomEventHandler<FConcertReplication_ChangeSyncControl>(
+					[this, &bReceivedEvent](const FConcertSessionContext&, const FConcertReplication_ChangeSyncControl& Event)
+					{
+						bReceivedEvent = true;
+						TestEqual(TEXT("Event.NewControlStates.Num() == 1"), Event.NewControlStates.Num(), 1);
+
+						const bool* bState = Event.NewControlStates.Find({ ObjectStreamId, ObjectReplicator->TestObject });
+						TestTrue(TEXT("Gained sync control"), bState && *bState);
+					});
+				
+				Receiver->GetClientReplicationManager()
+					.UnmuteObjects({ ObjectReplicator->TestObject })
+					.Next([this](FConcertReplication_ChangeMuteState_Response&& Response)
+					{
+						TestTrue(TEXT("Success"), Response.IsSuccess());
+						// Receiver has nothing registered. Make sure it does not receive the data intended for Sender.
+						TestTrue(TEXT("SyncControl.IsEmpty()"), Response.SyncControl.IsEmpty());
+					});
+				TestTrue(TEXT("Received event"), bReceivedEvent);
+			});
+		});
+	}
+}
\ No newline at end of file
diff --git a/Engine/Plugins/Developer/Concert/ConcertSync/ConcertSyncTest/Source/ConcertSyncTest/Private/Replication/Muting/MutingRejection.spec.cpp b/Engine/Plugins/Developer/Concert/ConcertSync/ConcertSyncTest/Source/ConcertSyncTest/Private/Replication/Muting/MutingRejection.spec.cpp
new file mode 100644
--- /dev/null
+++ b/Engine/Plugins/Developer/Concert/ConcertSync/ConcertSyncTest/Source/ConcertSyncTest/Private/Replication/Muting/MutingRejection.spec.cpp
@@ -0,0 +1,257 @@
+﻿// Copyright Epic Games, Inc. All Rights Reserved.
+
+#include "Misc/AutomationTest.h"
+#include "Replication/Util/Spec/ObjectTestReplicator.h"
+#include "Replication/Util/Spec/ReplicationClient.h"
+#include "Replication/Util/Spec/ReplicationServer.h"
+
+namespace UE::ConcertSyncTests::Replication
+{
+	/** Tests that the rejection cases outlined in FConcertReplication_ChangeMuteState_Request work. */
+	BEGIN_DEFINE_SPEC(FMutingRejectionSpec, "Editor.Concert.Replication.Muting", EAutomationTestFlags::EditorContext | EAutomationTestFlags::EngineFilter)
+		TSharedPtr<FReplicationServer> Server;
+		FReplicationClient* Sender = nullptr;
+		FReplicationClient* Receiver = nullptr;
+	
+		// Leverage FObjectTestReplicator to create more UObjects
+		TSharedPtr<FObjectTestReplicator> ObjectReplicator_Foo;			
+		TSharedPtr<FObjectTestReplicator> Replicator_FooSubobject;
+
+		void ValidateRejection(const FConcertReplication_ChangeMuteState_Response& Response)
+		{
+			TestTrue(TEXT("Request failed"), Response.IsFailure());
+			TestTrue(TEXT("Request rejected"), Response.ErrorCode == EConcertReplicationMuteErrorCode::Rejected);
+			TestEqual(TEXT("RejectionReason.Num() == 1"), Response.RejectionReasons.Num(), 1);
+			TestTrue(TEXT("RejectionReasons.Contains(UnregisteredObject)"), Response.RejectionReasons.Contains(ObjectReplicator_Foo->TestObject));
+		}
+	END_DEFINE_SPEC(FMutingRejectionSpec);
+
+	void FMutingRejectionSpec::Define()
+	{
+		BeforeEach([this]
+		{
+			Server		= MakeShared<FReplicationServer>(*this);
+			Sender		= &Server->ConnectClient();
+			Receiver	= &Server->ConnectClient();
+
+			ObjectReplicator_Foo			= MakeShared<FObjectTestReplicator>();
+			Replicator_FooSubobject			= ObjectReplicator_Foo->CreateSubobjectReplicator();
+		});
+		AfterEach([this]
+		{
+			// Test would hold onto this for rest of engine lifetime. Clean up this mini would-be leak.
+			Server.Reset();
+			ObjectReplicator_Foo.Reset();
+			Replicator_FooSubobject.Reset();
+		});
+
+		Describe("Request rejection cases", [this]()
+		{
+			It("Reject: Explicitly muting unreferenced object (EConcertReplicationMuteFlags::None)", [this]()
+			{
+				Sender->JoinReplication();
+				Sender->GetClientReplicationManager()
+					.MuteObjects({ ObjectReplicator_Foo->TestObject}, EConcertReplicationMuteOption::OnlyObject)
+					.Next([this](FConcertReplication_ChangeMuteState_Response&& Response)
+					{
+						ValidateRejection(Response);
+					});
+			});
+
+			It("Reject: Explicitly unmuting unreferenced object (EConcertReplicationMuteFlags::None)", [this]()
+			{
+				Sender->JoinReplication();
+				Sender->GetClientReplicationManager()
+					.UnmuteObjects({ ObjectReplicator_Foo->TestObject}, EConcertReplicationMuteOption::OnlyObject)
+					.Next([this](FConcertReplication_ChangeMuteState_Response&& Response)
+					{
+						ValidateRejection(Response);
+					});
+			});
+
+			It("Reject: Explicitly muting unreferenced object (EConcertReplicationMuteFlags::IncludeSubobjects) without subobject", [this]()
+			{
+				Sender->JoinReplication();
+				Sender->GetClientReplicationManager()
+					.MuteObjects({ ObjectReplicator_Foo->TestObject}, EConcertReplicationMuteOption::ObjectAndSubobjects)
+					.Next([this](FConcertReplication_ChangeMuteState_Response&& Response)
+					{
+						ValidateRejection(Response);
+					});
+			});
+
+			It("Reject: Explicitly unmuting unreferenced object (EConcertReplicationMuteFlags::IncludeSubobjects) without subobject", [this]()
+			{
+				Sender->JoinReplication();
+				Sender->GetClientReplicationManager()
+					.UnmuteObjects({ ObjectReplicator_Foo->TestObject}, EConcertReplicationMuteOption::ObjectAndSubobjects)
+					.Next([this](FConcertReplication_ChangeMuteState_Response&& Response)
+					{
+						ValidateRejection(Response);
+					});
+			});
+			
+			It("Allow: Explicitly muting unreferenced object (EConcertReplicationMuteFlags::IncludeSubobjects) with subobject", [this]()
+			{
+				Sender->JoinReplication(Replicator_FooSubobject->CreateSenderArgs());
+				Sender->GetClientReplicationManager()
+					.MuteObjects({ ObjectReplicator_Foo->TestObject}, EConcertReplicationMuteOption::ObjectAndSubobjects)
+					.Next([this](FConcertReplication_ChangeMuteState_Response&& Response)
+					{
+						TestTrue(TEXT("Success"), Response.IsSuccess());
+					});
+
+				Sender->GetClientReplicationManager()
+					.QueryMuteState()
+					.Next([this](FConcertReplication_QueryMuteState_Response&& Response)
+					{
+						TestEqual(TEXT("ExplicitlyMutedObjects.Num() == 1"), Response.ExplicitlyMutedObjects.Num(), 1);
+						TestEqual(TEXT("ExplicitlyUnmutedObjects.Num() == 0"), Response.ExplicitlyUnmutedObjects.Num(), 0);
+						TestEqual(TEXT("ImplicitlyMutedObjects.Num() == 1"), Response.ImplicitlyMutedObjects.Num(), 1);
+						TestEqual(TEXT("ImplicitlyUnmutedObjects.Num() == 0"), Response.ImplicitlyUnmutedObjects.Num(), 0);
+
+						const FConcertReplication_ObjectMuteSetting* FooMuteSetting = Response.ExplicitlyMutedObjects.Find(ObjectReplicator_Foo->TestObject);
+						TestTrue(TEXT("ExplicitlyMutedObjects.Contains(Subobject)"), FooMuteSetting && FooMuteSetting->Flags == EConcertReplicationMuteOption::ObjectAndSubobjects);
+						TestTrue(TEXT("ImplicitlyMutedObjects.Contains(Subobject)"), Response.ImplicitlyMutedObjects.Contains(Replicator_FooSubobject->TestObject));
+					});
+			});
+			
+			It("Allow: Explicitly unmuting unreferenced object (EConcertReplicationMuteFlags::IncludeSubobjects) with subobject", [this]()
+			{
+				Sender->JoinReplication(Replicator_FooSubobject->CreateSenderArgs());
+
+				// Check nothing is muted if ...
+				Sender->GetClientReplicationManager().MuteObjects({ ObjectReplicator_Foo->TestObject}, EConcertReplicationMuteOption::ObjectAndSubobjects);
+				// ... we unmute with EConcertReplicationMuteFlags::IncludeSubobjects flag
+				Sender->GetClientReplicationManager().UnmuteObjects({ ObjectReplicator_Foo->TestObject}, EConcertReplicationMuteOption::ObjectAndSubobjects);
+				Sender->GetClientReplicationManager()
+					.QueryMuteState()
+					.Next([this](FConcertReplication_QueryMuteState_Response&& Response)
+					{
+						TestTrue(TEXT("Nothing  muted"), Response.IsEmpty());
+					});
+
+				// Check nothing is muted if ...
+				Sender->GetClientReplicationManager().MuteObjects({ ObjectReplicator_Foo->TestObject}, EConcertReplicationMuteOption::ObjectAndSubobjects);
+				// ... we unmute with EConcertReplicationMuteFlags::None flag
+				Sender->GetClientReplicationManager().UnmuteObjects({ ObjectReplicator_Foo->TestObject}, EConcertReplicationMuteOption::OnlyObject);
+				Sender->GetClientReplicationManager()
+					.QueryMuteState()
+					.Next([this](FConcertReplication_QueryMuteState_Response&& Response)
+					{
+						TestTrue(TEXT("Nothing  muted"), Response.IsEmpty());
+					});
+			});
+
+			It("Reject: Mute and unmute in same opereation", [this]()
+			{
+				Sender->JoinReplication(Replicator_FooSubobject->CreateSenderArgs());
+				FConcertReplication_ChangeMuteState_Request Request;
+				Request.ObjectsToMute.Add(ObjectReplicator_Foo->TestObject);
+				Request.ObjectsToUnmute.Add(ObjectReplicator_Foo->TestObject);
+				Sender->GetClientReplicationManager()
+					.ChangeMuteState(Request)
+					.Next([this](FConcertReplication_ChangeMuteState_Response&& Response)
+					{
+						ValidateRejection(Response);
+					});
+			});
+		});
+	}
+	
+	/** This tests that muting requests are not served if EConcertSyncSessionFlags::ShouldAllowGlobalMuting is not set. */
+	BEGIN_DEFINE_SPEC(FMutingWithoutFlagSpec, "Editor.Concert.Replication.Muting", EAutomationTestFlags::EditorContext | EAutomationTestFlags::EngineFilter)
+		TSharedPtr<FReplicationServer> Server;
+		FReplicationClient* Sender = nullptr;
+
+		TSharedPtr<FObjectTestReplicator> ObjectReplicator;
+	END_DEFINE_SPEC(FMutingWithoutFlagSpec);
+
+	void FMutingWithoutFlagSpec::Define()
+	{
+		BeforeEach([this]
+		{
+			Server = MakeShared<FReplicationServer>(*this, EConcertSyncSessionFlags::Default_MultiUserSession & ~EConcertSyncSessionFlags::ShouldAllowGlobalMuting);
+			Sender = &Server->ConnectClient();
+			ObjectReplicator = MakeShared<FObjectTestReplicator>();
+			
+			Sender->JoinReplication(ObjectReplicator->CreateSenderArgs());
+		});
+		AfterEach([this]
+		{
+			// Test would hold onto this for rest of engine lifetime. Clean up this mini would-be leak.
+			Server.Reset();
+			ObjectReplicator.Reset();
+		});
+
+		Describe("When EConcertSyncSessionFlags::ShouldAllowGlobalMuting is not set", [this]()
+		{
+			Describe("Mute request is not sent via IConcertClientReplicationManager API", [this]()
+			{
+				It("FConcertReplication_ChangeMuteState_Request is not sent to server", [this]
+				{
+					Server->GetServerSessionMock()->RegisterCustomRequestHandler<FConcertReplication_ChangeMuteState_Request, FConcertReplication_ChangeMuteState_Response>(
+						[this](const FConcertSessionContext&, const FConcertReplication_ChangeMuteState_Request&, FConcertReplication_ChangeMuteState_Response&)
+						{
+							AddError(TEXT("Client was not supposed to send the request to the server."));
+							return EConcertSessionResponseCode::InvalidRequest;
+						});
+					
+					Sender->GetClientReplicationManager().MuteObjects({ ObjectReplicator->TestObject })
+						.Next([this](FConcertReplication_ChangeMuteState_Response&& Response)
+						{
+							TestTrue(TEXT("Request rejected"), Response.ErrorCode == EConcertReplicationMuteErrorCode::Rejected);
+						});
+				});
+				
+				It("FConcertReplication_QueryMuteState_Request is not sent to server", [this]
+				{
+					Server->GetServerSessionMock()->RegisterCustomRequestHandler<FConcertReplication_QueryMuteState_Request, FConcertReplication_QueryMuteState_Response>(
+						[this](const FConcertSessionContext&, const FConcertReplication_QueryMuteState_Request&, FConcertReplication_QueryMuteState_Response&)
+						{
+							AddError(TEXT("Client was not supposed to send the request to the server."));
+							return EConcertSessionResponseCode::InvalidRequest;
+						});
+					
+					Sender->GetClientReplicationManager().QueryMuteState()
+						.Next([this](FConcertReplication_QueryMuteState_Response&& Response)
+						{
+							TestTrue(TEXT("Response.IsEmpty()"), Response.IsEmpty());
+						});
+				});
+			});
+			
+			Describe("Server rejects", [this]()
+			{
+				It("FConcertReplication_ChangeMuteState_Request", [this]
+				{
+					FConcertReplication_ChangeMuteState_Request Request;
+					Request.ObjectsToMute.Add(ObjectReplicator->TestObject, { EConcertReplicationMuteOption::ObjectAndSubobjects });
+					
+					Sender->GetClientSessionMock()
+						->SendCustomRequest<FConcertReplication_ChangeMuteState_Request, FConcertReplication_ChangeMuteState_Response>(Request, Sender->GetClientSessionMock()->GetSessionServerEndpointId())
+						.Next([this](FConcertReplication_ChangeMuteState_Response&& Response)
+						{
+							TestTrue(TEXT("Response.IsFailure()"), Response.IsFailure());
+							// The server's request handler returns EConcertSessionResponseCode::Failed, which makes Concert default construct the response hence timeout
+							TestTrue(TEXT("Response.ErrorCode == Timeout"), Response.ErrorCode == EConcertReplicationMuteErrorCode::Timeout);
+						});
+				});
+
+				It("FConcertReplication_QueryMuteState_Request", [this]
+				{
+					FConcertReplication_QueryMuteState_Request Request;
+					
+					Sender->GetClientSessionMock()
+						->SendCustomRequest<FConcertReplication_QueryMuteState_Request, FConcertReplication_QueryMuteState_Response>(Request, Sender->GetClientSessionMock()->GetSessionServerEndpointId())
+						.Next([this](FConcertReplication_QueryMuteState_Response&& Response)
+						{
+							// The server's request handler returns EConcertSessionResponseCode::Failed, which makes Concert default construct the response
+							TestTrue(TEXT("Response.IsFailure()"), Response.IsFailure());
+							TestTrue(TEXT("Response.IsEmpty()"), Response.IsEmpty());
+						});
+				});
+			});
+		});
+	}
+}
\ No newline at end of file
diff --git a/Engine/Plugins/Developer/Concert/ConcertSync/ConcertSyncTest/Source/ConcertSyncTest/Private/Replication/Muting/MutingReplication.spec.cpp b/Engine/Plugins/Developer/Concert/ConcertSync/ConcertSyncTest/Source/ConcertSyncTest/Private/Replication/Muting/MutingReplication.spec.cpp
new file mode 100644
--- /dev/null
+++ b/Engine/Plugins/Developer/Concert/ConcertSync/ConcertSyncTest/Source/ConcertSyncTest/Private/Replication/Muting/MutingReplication.spec.cpp
@@ -0,0 +1,239 @@
+﻿// Copyright Epic Games, Inc. All Rights Reserved.
+
+#include "Misc/AutomationTest.h"
+#include "Replication/Util/ConcertClientReplicationBridgeMock.h"
+#include "Replication/Util/Spec/ObjectTestReplicator.h"
+#include "Replication/Util/Spec/ReplicationClient.h"
+#include "Replication/Util/Spec/ReplicationServer.h"
+
+namespace UE::ConcertSyncTests::Replication
+{
+	/** This tests that replication works correctly after muting & unmuting. */
+	BEGIN_DEFINE_SPEC(FMutingReplicationSpec, "Editor.Concert.Replication.Muting", EAutomationTestFlags::EditorContext | EAutomationTestFlags::EngineFilter)
+	TSharedPtr<FReplicationServer> Server;
+	FReplicationClient* Sender = nullptr;
+	FReplicationClient* Receiver = nullptr;
+
+	FGuid StreamId_Foo					= FGuid::NewGuid();
+	FGuid StreamId_FooSubobject			= FGuid::NewGuid();
+	FGuid StreamId_Bar					= FGuid::NewGuid();
+	FGuid StreamId_BarSubobject			= FGuid::NewGuid();
+	FGuid StreamId_BarNestedSubobject	= FGuid::NewGuid();
+	// Leverage FObjectTestReplicator to create more UObjects
+	TSharedPtr<FObjectTestReplicator> Replicator_Foo;			
+	TSharedPtr<FObjectTestReplicator> Replicator_FooSubobject;		
+	TSharedPtr<FObjectTestReplicator> Replicator_Bar;			
+	TSharedPtr<FObjectTestReplicator> Replicator_BarSubobject;		
+	TSharedPtr<FObjectTestReplicator> Replicator_BarNestedSubobject;
+
+	FObjectReplicationContext MakeSenderToReceiverContext(const TCHAR* Context = nullptr) const { return { *Sender, *Server, *Receiver, Context }; }
+	END_DEFINE_SPEC(FMutingReplicationSpec);
+
+	/** This tests that muting requests work when EConcertSyncSessionFlags::ShouldAllowGlobalMuting is set. */
+	void FMutingReplicationSpec::Define()
+	{
+		BeforeEach([this]
+		{
+			Server		= MakeShared<FReplicationServer>(*this);
+			Sender		= &Server->ConnectClient();
+			Receiver	= &Server->ConnectClient();
+
+			Replicator_Foo					= MakeShared<FObjectTestReplicator>(TEXT("Foo"));
+			Replicator_FooSubobject			= Replicator_Foo->CreateSubobjectReplicator(TEXT("FooSubobject"));
+			Replicator_Bar					= MakeShared<FObjectTestReplicator>(TEXT("Bar"));
+			Replicator_BarSubobject			= Replicator_Bar->CreateSubobjectReplicator(TEXT("BarSubobject"));
+			Replicator_BarNestedSubobject	= Replicator_BarSubobject->CreateSubobjectReplicator(TEXT("BarNestedSubobject"));
+		});
+		AfterEach([this]
+		{
+			// Test would hold onto this for rest of engine lifetime. Clean up this mini would-be leak.
+			Server.Reset();
+			Replicator_Foo.Reset();
+			Replicator_FooSubobject.Reset();
+			Replicator_Bar.Reset();
+			Replicator_BarSubobject.Reset();
+			Replicator_BarNestedSubobject.Reset();
+		});
+
+		Describe("Replication with mute & unmute", [this]()
+		{
+			BeforeEach([this]()
+			{
+				ConcertSyncClient::Replication::FJoinReplicatedSessionArgs SenderJoinArgs;
+				SenderJoinArgs.Streams.Add(Replicator_Foo->CreateStream(StreamId_Foo));
+				SenderJoinArgs.Streams.Add(Replicator_FooSubobject->CreateStream(StreamId_FooSubobject));
+				SenderJoinArgs.Streams.Add(Replicator_Bar->CreateStream(StreamId_Bar));
+				SenderJoinArgs.Streams.Add(Replicator_BarSubobject->CreateStream(StreamId_BarSubobject));
+				SenderJoinArgs.Streams.Add(Replicator_BarNestedSubobject->CreateStream(StreamId_BarNestedSubobject));
+				
+				Sender->JoinReplication(SenderJoinArgs);
+				// Tells server intent to replicate the object
+				Sender->GetClientReplicationManager().TakeAuthorityOver({
+					Replicator_Foo->TestObject,
+					Replicator_FooSubobject->TestObject,
+					Replicator_Bar->TestObject,
+					Replicator_BarSubobject->TestObject,
+					Replicator_BarNestedSubobject->TestObject
+					});
+
+				// This gives sync control meaning that now the object can be replicated.
+				Receiver->JoinReplicationAsListener({ Replicator_Foo->TestObject, Replicator_FooSubobject->TestObject, Replicator_Bar->TestObject, Replicator_BarSubobject->TestObject, Replicator_BarNestedSubobject->TestObject });
+
+				bool bReceivedEvent = false;
+				FConcertReplication_ChangeMuteState_Request MuteRequest;
+				MuteRequest.ObjectsToMute.Add(Replicator_Foo->TestObject, { EConcertReplicationMuteOption::ObjectAndSubobjects });
+				MuteRequest.ObjectsToMute.Add(Replicator_Bar->TestObject, { EConcertReplicationMuteOption::OnlyObject });
+				Sender->GetClientReplicationManager()
+					.ChangeMuteState(MuteRequest)
+					.Next([this, &bReceivedEvent](FConcertReplication_ChangeMuteState_Response&& Response)
+					{
+						bReceivedEvent = true;
+						TestTrue(TEXT("Mute request successful"), Response.IsSuccess());
+					});
+				TestTrue(TEXT("Mute response received"), bReceivedEvent);
+			});
+
+			Describe("Don't replicate: Muted objects", [this]()
+			{
+				It("Foo was not replicated", [this]()
+				{
+					Sender->GetBridgeMock().InjectAvailableObject(*Replicator_Foo->TestObject);
+					Replicator_Foo->SimulateSendObjectToReceiver(*this, MakeSenderToReceiverContext(), { StreamId_Foo });
+					Replicator_Foo->TestValuesWereNotReplicated(*this);
+				});
+				
+				It("FooSubobject was not replicated", [this]()
+				{
+					Sender->GetBridgeMock().InjectAvailableObject(*Replicator_FooSubobject->TestObject);
+					Replicator_FooSubobject->SimulateSendObjectToReceiver(*this, MakeSenderToReceiverContext(), { StreamId_FooSubobject });
+					Replicator_FooSubobject->TestValuesWereNotReplicated(*this);
+				});
+				It("Bar was not replicated", [this]()
+				{
+					Sender->GetBridgeMock().InjectAvailableObject(*Replicator_Bar->TestObject);
+					Replicator_Bar->SimulateSendObjectToReceiver(*this, MakeSenderToReceiverContext(), { StreamId_Bar });
+					Replicator_Bar->TestValuesWereNotReplicated(*this);
+				});
+				
+				// This subobject is not affected
+				It("BarSubobject was replicated", [this]()
+				{
+					Sender->GetBridgeMock().InjectAvailableObject(*Replicator_BarSubobject->TestObject);
+					Replicator_BarSubobject->SimulateSendObjectToReceiver(*this, MakeSenderToReceiverContext(), { StreamId_BarSubobject });
+					Replicator_BarSubobject->TestValuesWereReplicated(*this);
+				});
+			});
+			
+			Describe("Replicate: Muted objects after unmuting", [this]()
+			{
+				BeforeEach([this]()
+				{
+					bool bReceivedEvent = false;
+					Sender->GetClientReplicationManager()
+						.UnmuteObjects({ Replicator_Foo->TestObject })
+						.Next([this, &bReceivedEvent](FConcertReplication_ChangeMuteState_Response&& Response)
+						{
+							bReceivedEvent = true;
+							TestTrue(TEXT("Unmute request successful"), Response.ErrorCode == EConcertReplicationMuteErrorCode::Accepted);
+						});
+					TestTrue(TEXT("Mute response received"), bReceivedEvent);
+				});
+				
+				It("Foo was replicated", [this]()
+				{
+					Sender->GetBridgeMock().InjectAvailableObject(*Replicator_Foo->TestObject);
+					Replicator_Foo->SimulateSendObjectToReceiver(*this, MakeSenderToReceiverContext(), { StreamId_Foo });
+					Replicator_Foo->TestValuesWereReplicated(*this);
+				});
+				
+				It("FooSubobject was replicated", [this]()
+				{
+					Sender->GetBridgeMock().InjectAvailableObject(*Replicator_FooSubobject->TestObject);
+					Replicator_FooSubobject->SimulateSendObjectToReceiver(*this, MakeSenderToReceiverContext(), { StreamId_FooSubobject });
+					Replicator_FooSubobject->TestValuesWereReplicated(*this);
+				});
+				It("Bar was not replicated", [this]()
+				{
+					Sender->GetBridgeMock().InjectAvailableObject(*Replicator_Bar->TestObject);
+					Replicator_Bar->SimulateSendObjectToReceiver(*this, MakeSenderToReceiverContext(), { StreamId_Bar });
+					Replicator_Bar->TestValuesWereNotReplicated(*this);
+				});
+				
+				// This subobject is not affected
+				It("BarSubobject was replicated", [this]()
+				{
+					Sender->GetBridgeMock().InjectAvailableObject(*Replicator_BarSubobject->TestObject);
+					Replicator_BarSubobject->SimulateSendObjectToReceiver(*this, MakeSenderToReceiverContext(), { StreamId_BarSubobject });
+					Replicator_BarSubobject->TestValuesWereReplicated(*this);
+				});
+			});
+
+			It("Replicate: subobject when unmuted explicitly", [this]()
+			{
+				Sender->GetClientReplicationManager().UnmuteObjects({ Replicator_Foo->TestObject }, EConcertReplicationMuteOption::OnlyObject);
+				
+				Sender->GetBridgeMock().InjectAvailableObject(*Replicator_FooSubobject->TestObject);
+				Replicator_FooSubobject->SimulateSendObjectToReceiver(*this, MakeSenderToReceiverContext(), { StreamId_FooSubobject });
+				Replicator_FooSubobject->TestValuesWereReplicated(*this);
+			});
+
+			Describe("Replicate: remove stream that contain muted parent.", [this]()
+			{
+				BeforeEach([this]()
+				{
+					const FConcertReplication_ChangeStream_Request Request { .StreamsToRemove = { StreamId_Foo } };
+					Sender->GetClientReplicationManager().ChangeStream(Request);
+					Sender->GetBridgeMock().InjectAvailableObject(*Replicator_FooSubobject->TestObject);
+				});
+
+				It("Don't replicate: subobject stays implicitly muted", [this]()
+				{
+					Replicator_FooSubobject->SimulateSendObjectToReceiver(*this, MakeSenderToReceiverContext(), { StreamId_FooSubobject });
+					Replicator_FooSubobject->TestValuesWereNotReplicated(*this);
+				});
+
+				It("Replicate: Explicitly unmuting parent still unmutes children", [this]()
+				{
+					Sender->GetClientReplicationManager().UnmuteObjects({ Replicator_Foo->TestObject });
+					
+					Replicator_FooSubobject->SimulateSendObjectToReceiver(*this, MakeSenderToReceiverContext(), { StreamId_FooSubobject });
+					Replicator_FooSubobject->TestValuesWereReplicated(*this);
+				});
+			});
+
+			Describe("Replicate: implicitly unmuted objects", [this]()
+			{
+				BeforeEach([this]()
+				{
+					IConcertClientReplicationManager& ReplicationManager = Sender->GetClientReplicationManager();
+					ReplicationManager.MuteObjects({ Replicator_Bar->TestObject }, EConcertReplicationMuteOption::ObjectAndSubobjects);
+					ReplicationManager.UnmuteObjects({ Replicator_BarSubobject->TestObject }, EConcertReplicationMuteOption::ObjectAndSubobjects);
+				});
+
+				// Explicitly muted	 > does not replicate
+				It("Bar was not replicated", [this]()
+				{
+					Sender->GetBridgeMock().InjectAvailableObject(*Replicator_Bar->TestObject);
+					Replicator_Bar->SimulateSendObjectToReceiver(*this, MakeSenderToReceiverContext(), { StreamId_Bar });
+					Replicator_Bar->TestValuesWereNotReplicated(*this);
+				});
+				
+				// Explicitly unmuted > does replicate
+				It("BarSubobject was replicated", [this]()
+				{
+					Sender->GetBridgeMock().InjectAvailableObject(*Replicator_BarSubobject->TestObject);
+					Replicator_BarSubobject->SimulateSendObjectToReceiver(*this, MakeSenderToReceiverContext(), { StreamId_BarSubobject });
+					Replicator_BarSubobject->TestValuesWereReplicated(*this);
+				});
+
+				// Implicitly unmuted > does replicate
+				It("BarNestedSubobject was replicated", [this]()
+				{
+					Sender->GetBridgeMock().InjectAvailableObject(*Replicator_BarNestedSubobject->TestObject);
+					Replicator_BarNestedSubobject->SimulateSendObjectToReceiver(*this, MakeSenderToReceiverContext(), { StreamId_BarNestedSubobject });
+					Replicator_BarNestedSubobject->TestValuesWereReplicated(*this);
+				});
+			});
+		});
+	}
+}
\ No newline at end of file
diff --git a/Engine/Plugins/Developer/Concert/ConcertSync/ConcertSyncTest/Source/ConcertSyncTest/Private/Replication/ReplicationTestInterface.h b/Engine/Plugins/Developer/Concert/ConcertSync/ConcertSyncTest/Source/ConcertSyncTest/Private/Replication/ReplicationTestInterface.h
--- a/Engine/Plugins/Developer/Concert/ConcertSync/ConcertSyncTest/Source/ConcertSyncTest/Private/Replication/ReplicationTestInterface.h
+++ b/Engine/Plugins/Developer/Concert/ConcertSync/ConcertSyncTest/Source/ConcertSyncTest/Private/Replication/ReplicationTestInterface.h
@@ -2,6 +2,7 @@
 
 #pragma once
 
+#include "ConcertSyncSessionFlags.h"
 #include "Templates/SharedPointer.h"
 
 class IConcertClientSession;
@@ -18,7 +19,8 @@ namespace UE::ConcertSyncClient::TestInterface
 {
 	extern CONCERTSYNCCLIENT_API TSharedRef<IConcertClientReplicationManager> CreateClientReplicationManager(
 		TSharedRef<IConcertClientSession> InLiveSession,
-		IConcertClientReplicationBridge& InBridge UE_LIFETIMEBOUND
+		IConcertClientReplicationBridge& InBridge UE_LIFETIMEBOUND,
+		EConcertSyncSessionFlags SessionFlags = EConcertSyncSessionFlags::Default_MultiUserSession
 		);
 
 	extern CONCERTSYNCCLIENT_API TSharedRef<IConcertClientReplicationBridge> CreateClientReplicationBridge();
@@ -27,6 +29,7 @@ namespace UE::ConcertSyncClient::TestInterface
 namespace UE::ConcertSyncServer::TestInterface
 {
 	extern CONCERTSYNCSERVER_API TSharedRef<Replication::IConcertServerReplicationManager> CreateServerReplicationManager(
-		TSharedRef<IConcertServerSession> InLiveSession
+		TSharedRef<IConcertServerSession> InLiveSession,
+		EConcertSyncSessionFlags InSessionFlags = EConcertSyncSessionFlags::Default_MultiUserSession
 		);
 }
\ No newline at end of file
diff --git a/Engine/Plugins/Developer/Concert/ConcertSync/ConcertSyncTest/Source/ConcertSyncTest/Private/Replication/Util/Spec/ObjectTestReplicator.cpp b/Engine/Plugins/Developer/Concert/ConcertSync/ConcertSyncTest/Source/ConcertSyncTest/Private/Replication/Util/Spec/ObjectTestReplicator.cpp
--- a/Engine/Plugins/Developer/Concert/ConcertSync/ConcertSyncTest/Source/ConcertSyncTest/Private/Replication/Util/Spec/ObjectTestReplicator.cpp
+++ b/Engine/Plugins/Developer/Concert/ConcertSync/ConcertSyncTest/Source/ConcertSyncTest/Private/Replication/Util/Spec/ObjectTestReplicator.cpp
@@ -9,9 +9,24 @@
 
 namespace UE::ConcertSyncTests::Replication
 {
-	TSharedRef<FObjectTestReplicator> FObjectTestReplicator::CreateSubobjectReplicator() const
+	namespace Private
 	{
-		UTestReflectionObject* Subobject = NewObject<UTestReflectionObject>(TestObject);
+		FString TestName(const TCHAR* BaseName, const TCHAR* Context)
+		{
+			return Context
+				? FString::Printf(TEXT("%s: %s"), Context, BaseName)
+				: BaseName;
+		}
+		
+		FString TestName(const TCHAR* BaseName, const FObjectReplicationContext& Context)
+		{
+			return TestName(BaseName, Context.ContextName);
+		}
+	}
+
+	TSharedRef<FObjectTestReplicator> FObjectTestReplicator::CreateSubobjectReplicator(const FName BaseName) const
+	{
+		UTestReflectionObject* Subobject = NewObject<UTestReflectionObject>(TestObject, BaseName);
 		TestObject->InstancedSubobject = Subobject;
 		return MakeShared<FObjectTestReplicator>(Subobject);
 	}
@@ -52,29 +67,30 @@ namespace UE::ConcertSyncTests::Replication
 		EPropertyReplicationFlags PropertyFlags
 		) const
 	{
-		auto TestReplicationData_Server = [this, &Test, &SenderStreams, &OnServerReceive](const FConcertSessionContext& Context, const FConcertReplication_BatchReplicationEvent& Event)
+		using namespace Private;
+		const auto TestReplicationData_Server = [this, &Test, &Context, &SenderStreams, &OnServerReceive](const FConcertSessionContext& SessionContext, const FConcertReplication_BatchReplicationEvent& Event)
 		{
-			Test.TestEqual(TEXT("Server received right number of streams"), Event.Streams.Num(), SenderStreams.Num());
+			Test.TestEqual(TestName(TEXT("Server received right number of streams"), Context), Event.Streams.Num(), SenderStreams.Num());
 			for (int32 i = 0; i < Event.Streams.Num(); ++i)
 			{
-				Test.TestEqual(TEXT("Server received 1 object"), Event.Streams[i].ReplicatedObjects.Num() , 1);
-				Test.TestTrue(TEXT("Server received from correct stream"), SenderStreams.Contains(Event.Streams[i].StreamId));
+				Test.TestEqual(TestName(TEXT("Server received 1 object"), Context), Event.Streams[i].ReplicatedObjects.Num() , 1);
+				Test.TestTrue(TestName(TEXT("Server received from correct stream"), Context), SenderStreams.Contains(Event.Streams[i].StreamId));
 				const FSoftObjectPath ObjectPath = Event.Streams[i].ReplicatedObjects.IsEmpty() ? FSoftObjectPath{} : Event.Streams[i].ReplicatedObjects[0].ReplicatedObject;
-				Test.TestEqual(TEXT("Server's received object has correct path"), ObjectPath, FSoftObjectPath(TestObject));
+				Test.TestEqual(TestName(TEXT("Server's received object has correct path"), Context), ObjectPath, FSoftObjectPath(TestObject));
 			}
-			OnServerReceive(Context, Event);
+			OnServerReceive(SessionContext, Event);
 		};
-		auto TestReplicationData_Client_Receiver = [this, &Test, &SenderStreams, &OnReceiverClientReceive](const FConcertSessionContext& Context, const FConcertReplication_BatchReplicationEvent& Event)
+		const auto TestReplicationData_Client_Receiver = [this, &Test, &Context, &SenderStreams, &OnReceiverClientReceive](const FConcertSessionContext& SessionContext, const FConcertReplication_BatchReplicationEvent& Event)
 		{
-			Test.TestEqual(TEXT("Client 2 received right number of streams"), Event.Streams.Num(), SenderStreams.Num());
+			Test.TestEqual(TestName(TEXT("Client 2 received right number of streams"), Context), Event.Streams.Num(), SenderStreams.Num());
 			for (int32 i = 0; i < Event.Streams.Num(); ++i)
 			{
-				Test.TestEqual(TEXT("Client 2 received 1 object"), Event.Streams[i].ReplicatedObjects.Num() , 1);
-				Test.TestTrue(TEXT("Client 2 received from correct stream"), SenderStreams.Contains(Event.Streams[i].StreamId));
+				Test.TestEqual(TestName(TEXT("Client 2 received 1 object"), Context), Event.Streams[i].ReplicatedObjects.Num() , 1);
+				Test.TestTrue(TestName(TEXT("Client 2 received from correct stream"), Context), SenderStreams.Contains(Event.Streams[i].StreamId));
 				const FSoftObjectPath ObjectPath = Event.Streams[i].ReplicatedObjects.IsEmpty() ? FSoftObjectPath{} : Event.Streams[i].ReplicatedObjects[0].ReplicatedObject;
-				Test.TestEqual(TEXT("Client 2's received object has correct path"), ObjectPath, FSoftObjectPath(TestObject));
+				Test.TestEqual(TestName(TEXT("Client 2's received object has correct path"), Context), ObjectPath, FSoftObjectPath(TestObject));
 			}
-			OnReceiverClientReceive(Context, Event);
+			OnReceiverClientReceive(SessionContext, Event);
 		};
 		const FDelegateHandle ServerHandle = Context.Server.GetServerSessionMock()->RegisterCustomEventHandler<FConcertReplication_BatchReplicationEvent>(TestReplicationData_Server);
 		const FDelegateHandle ClientHandle = Context.Receiver.GetClientSessionMock()->RegisterCustomEventHandler<FConcertReplication_BatchReplicationEvent>(TestReplicationData_Client_Receiver);
@@ -122,34 +138,38 @@ namespace UE::ConcertSyncTests::Replication
 		}
 	}
 
-	void FObjectTestReplicator::TestValuesWereReplicated(FAutomationTestBase& Test, EPropertyReplicationFlags PropertyFlags) const
+	void FObjectTestReplicator::TestValuesWereReplicated(FAutomationTestBase& Test, EPropertyReplicationFlags PropertyFlags, const TCHAR* Context) const
 	{
+		using namespace Private;
+		
 		const bool bSendCDOValues = EnumHasAnyFlags(PropertyFlags, EPropertyReplicationFlags::SendCDOValues);
 		if (EnumHasAnyFlags(PropertyFlags, EPropertyReplicationFlags::Float))
 		{
 			const float ExpectedValue = bSendCDOValues ? GetMutableDefault<UTestReflectionObject>()->Float : SentFloat;
-			const bool bWasReplicated = Test.TestEqual(TEXT("Float"), TestObject->Float, ExpectedValue);
-			Test.AddErrorIfFalse(bWasReplicated, TEXT("Failed to replicate \"Float\" property"));
+			const bool bWasReplicated = Test.TestEqual(TestName(TEXT("Float"), Context), TestObject->Float, ExpectedValue);
+			Test.AddErrorIfFalse(bWasReplicated, TestName(TEXT("Failed to replicate \"Float\" property"), Context));
 		}
 		if (EnumHasAnyFlags(PropertyFlags, EPropertyReplicationFlags::Vector))
 		{
 			const FVector ExpectedValue = bSendCDOValues ? GetMutableDefault<UTestReflectionObject>()->Vector : SentVector;
-			const bool bWasReplicated = Test.TestEqual(TEXT("Vector"), TestObject->Vector, ExpectedValue);
-			Test.AddErrorIfFalse(bWasReplicated, TEXT("Failed to replicate \"Vector\" property"));
+			const bool bWasReplicated = Test.TestEqual(TestName(TEXT("Vector"), Context), TestObject->Vector, ExpectedValue);
+			Test.AddErrorIfFalse(bWasReplicated, TestName(TEXT("Failed to replicate \"Vector\" property"), Context));
 		}
 	}
 
-	void FObjectTestReplicator::TestValuesWereNotReplicated(FAutomationTestBase& Test, EPropertyReplicationFlags PropertyFlags) const
+	void FObjectTestReplicator::TestValuesWereNotReplicated(FAutomationTestBase& Test, EPropertyReplicationFlags PropertyFlags, const TCHAR* Context) const
 	{
+		using namespace Private;
+		
 		if (EnumHasAnyFlags(PropertyFlags, EPropertyReplicationFlags::Float))
 		{
-			const bool bWasNotReplicated = Test.TestEqual(TEXT("Float"), TestObject->Float, DifferentFloat);
-			Test.AddErrorIfFalse(bWasNotReplicated, TEXT("Probably \"Float\" property was replicated even though it was not supposed to be!"));
+			const bool bWasNotReplicated = Test.TestEqual(TestName(TEXT("Float"), Context), TestObject->Float, DifferentFloat);
+			Test.AddErrorIfFalse(bWasNotReplicated, TestName(TEXT("Probably \"Float\" property was replicated even though it was not supposed to be!"), Context));
 		}
 		if (EnumHasAnyFlags(PropertyFlags, EPropertyReplicationFlags::Vector))
 		{
-			const bool bWasNotReplicated = Test.TestEqual(TEXT("Vector"), TestObject->Vector, DifferentVector);
-			Test.AddErrorIfFalse(bWasNotReplicated, TEXT("Probably \"Vector\" property was replicated even though it was not supposed to be!"));
+			const bool bWasNotReplicated = Test.TestEqual(TestName(TEXT("Vector"), Context), TestObject->Vector, DifferentVector);
+			Test.AddErrorIfFalse(bWasNotReplicated, TestName(TEXT("Probably \"Vector\" property was replicated even though it was not supposed to be!"), Context));
 		}
 	}
 }
diff --git a/Engine/Plugins/Developer/Concert/ConcertSync/ConcertSyncTest/Source/ConcertSyncTest/Private/Replication/Util/Spec/ObjectTestReplicator.h b/Engine/Plugins/Developer/Concert/ConcertSync/ConcertSyncTest/Source/ConcertSyncTest/Private/Replication/Util/Spec/ObjectTestReplicator.h
--- a/Engine/Plugins/Developer/Concert/ConcertSync/ConcertSyncTest/Source/ConcertSyncTest/Private/Replication/Util/Spec/ObjectTestReplicator.h
+++ b/Engine/Plugins/Developer/Concert/ConcertSync/ConcertSyncTest/Source/ConcertSyncTest/Private/Replication/Util/Spec/ObjectTestReplicator.h
@@ -37,6 +37,8 @@ namespace UE::ConcertSyncTests::Replication
 		FReplicationClient& Sender;
 		FReplicationServer& Server;
 		FReplicationClient& Receiver;
+		/** Helps debug failing tests. */
+		const TCHAR* ContextName = nullptr;
 	};
 	
 	/**
@@ -61,9 +63,10 @@ namespace UE::ConcertSyncTests::Replication
 
 		FObjectTestReplicator() = default;
 		explicit FObjectTestReplicator(UTestReflectionObject* TestObject) : TestObject(TestObject) {}
+		explicit FObjectTestReplicator(const FName BaseName) : TestObject(NewObject<UTestReflectionObject>(GetTransientPackage(), MakeUniqueObjectName(GetTransientPackage(), UTestReflectionObject::StaticClass(), BaseName))) {}
 
 		/** Create a subobject of TestObject, assigns it to TestObject->InstancedSubobject, and returns a replicator for replicating it. */
-		TSharedRef<FObjectTestReplicator> CreateSubobjectReplicator() const;
+		TSharedRef<FObjectTestReplicator> CreateSubobjectReplicator(const FName BaseName = NAME_None) const;
 
 		/** Util for creating join args for replicating TestObject. */
 		ConcertSyncClient::Replication::FJoinReplicatedSessionArgs CreateSenderArgs(
@@ -93,10 +96,14 @@ namespace UE::ConcertSyncTests::Replication
 		/** Sets the specified properties to values different from the test values */
 		void SetDifferentValues(EPropertyReplicationFlags PropertyFlags = EPropertyReplicationFlags::All) const;
 
+		
 		/** Tests that the specified properties are equal to their test values */
-		void TestValuesWereReplicated(FAutomationTestBase& Test, EPropertyReplicationFlags PropertyFlags = EPropertyReplicationFlags::All) const;
+		void TestValuesWereReplicated(FAutomationTestBase& Test, EPropertyReplicationFlags PropertyFlags = EPropertyReplicationFlags::All, const TCHAR* Context = nullptr) const;
 		/** Tests that the specified properties are equal to the values different from the test values (i.e. the values SetDifferentValues sets). */
-		void TestValuesWereNotReplicated(FAutomationTestBase& Test, EPropertyReplicationFlags PropertyFlags = EPropertyReplicationFlags::All) const;
+		void TestValuesWereNotReplicated(FAutomationTestBase& Test, EPropertyReplicationFlags PropertyFlags = EPropertyReplicationFlags::All, const TCHAR* Context = nullptr) const;
+		
+		void TestValuesWereReplicated(FAutomationTestBase& Test, const TCHAR* Context, EPropertyReplicationFlags PropertyFlags = EPropertyReplicationFlags::All) const { TestValuesWereReplicated(Test, PropertyFlags, Context); }
+		void TestValuesWereNotReplicated(FAutomationTestBase& Test, const TCHAR* Context, EPropertyReplicationFlags PropertyFlags = EPropertyReplicationFlags::All) const { TestValuesWereNotReplicated(Test, PropertyFlags, Context); }
 	};
 }
 
diff --git a/Engine/Plugins/Developer/Concert/ConcertSync/ConcertSyncTest/Source/ConcertSyncTest/Private/Replication/Util/Spec/ReplicationClient.cpp b/Engine/Plugins/Developer/Concert/ConcertSync/ConcertSyncTest/Source/ConcertSyncTest/Private/Replication/Util/Spec/ReplicationClient.cpp
--- a/Engine/Plugins/Developer/Concert/ConcertSync/ConcertSyncTest/Source/ConcertSyncTest/Private/Replication/Util/Spec/ReplicationClient.cpp
+++ b/Engine/Plugins/Developer/Concert/ConcertSync/ConcertSyncTest/Source/ConcertSyncTest/Private/Replication/Util/Spec/ReplicationClient.cpp
@@ -17,10 +17,10 @@ namespace UE::ConcertSyncTests::Replication
 
 	TFuture<ConcertSyncClient::Replication::FJoinReplicatedSessionResult> FReplicationClient::JoinReplication(
 		ConcertSyncClient::Replication::FJoinReplicatedSessionArgs Args,
-		EReplicationClientFlags Flags
+		EReplicationClientFlags TestFlags
 		)
 	{
-		if (EnumHasAnyFlags(Flags, EReplicationClientFlags::UseRealReplicationBridge))
+		if (EnumHasAnyFlags(TestFlags, EReplicationClientFlags::UseRealReplicationBridge))
 		{
 			BridgeUsed = ConcertSyncClient::TestInterface::CreateClientReplicationBridge();
 		}
@@ -29,7 +29,7 @@ namespace UE::ConcertSyncTests::Replication
 			BridgeMock =  MakeShared<FConcertClientReplicationBridgeMock>();
 			BridgeUsed = BridgeMock;
 		}
-		ClientReplicationManager = ConcertSyncClient::TestInterface::CreateClientReplicationManager(ClientSessionMock, *BridgeUsed);
+		ClientReplicationManager = ConcertSyncClient::TestInterface::CreateClientReplicationManager(ClientSessionMock, *BridgeUsed, SessionFlags);
 		
 		// 1.1 Sender offers to send all UTestReflectionObject properties
 		bool bSuccess = false;
@@ -40,7 +40,7 @@ namespace UE::ConcertSyncTests::Replication
 				return Result;
 			});
 
-		if (!EnumHasAnyFlags(Flags, EReplicationClientFlags::SkipJoinTest))
+		if (!EnumHasAnyFlags(TestFlags, EReplicationClientFlags::SkipJoinTest))
 		{
 			TestContext.TestTrue(TEXT("Replication joined successfully"), bSuccess);
 		}
diff --git a/Engine/Plugins/Developer/Concert/ConcertSync/ConcertSyncTest/Source/ConcertSyncTest/Private/Replication/Util/Spec/ReplicationClient.h b/Engine/Plugins/Developer/Concert/ConcertSync/ConcertSyncTest/Source/ConcertSyncTest/Private/Replication/Util/Spec/ReplicationClient.h
--- a/Engine/Plugins/Developer/Concert/ConcertSync/ConcertSyncTest/Source/ConcertSyncTest/Private/Replication/Util/Spec/ReplicationClient.h
+++ b/Engine/Plugins/Developer/Concert/ConcertSync/ConcertSyncTest/Source/ConcertSyncTest/Private/Replication/Util/Spec/ReplicationClient.h
@@ -2,6 +2,7 @@
 
 #pragma once
 
+#include "ConcertSyncSessionFlags.h"
 #include "Replication/IConcertClientReplicationManager.h"
 #include "Util/ClientServerCommunicationTest.h"
 #include "Templates/UnrealTemplate.h"
@@ -33,8 +34,9 @@ namespace UE::ConcertSyncTests::Replication
 	{
 	public:
 		
-		FReplicationClient(const FGuid& ClientEndPointId, FConcertServerSessionMock& Server, FAutomationTestBase& TestContext)
-			: TestContext(TestContext)
+		FReplicationClient(const FGuid& ClientEndPointId, EConcertSyncSessionFlags SessionFlags, FConcertServerSessionMock& Server, FAutomationTestBase& TestContext)
+			: SessionFlags(SessionFlags)
+			, TestContext(TestContext)
 			, ClientSessionMock(MakeShared<FConcertClientSessionMock>(ClientEndPointId, Server))
 		{}
 		
@@ -44,7 +46,7 @@ namespace UE::ConcertSyncTests::Replication
 		/** Joins the client into replication allowing them to participate in replication */
 		TFuture<ConcertSyncClient::Replication::FJoinReplicatedSessionResult> JoinReplication(
 			ConcertSyncClient::Replication::FJoinReplicatedSessionArgs Args = {},
-			EReplicationClientFlags Flags = EReplicationClientFlags::None
+			EReplicationClientFlags TestFlags = EReplicationClientFlags::None
 			);
 		/** This overload joins replication and injects the objects into the FConcertClientReplicationBridgeMock so they can be received. */
 		TFuture<ConcertSyncClient::Replication::FJoinReplicatedSessionResult> JoinReplicationAsListener(
@@ -53,11 +55,14 @@ namespace UE::ConcertSyncTests::Replication
 
 		TSharedRef<FConcertClientSessionBaseMock> GetClientSessionMock() const { return ClientSessionMock; }
 
-		FConcertClientReplicationBridgeMock& GetBridgeMock() const { return *BridgeMock; }
-		IConcertClientReplicationManager& GetClientReplicationManager() const { return *ClientReplicationManager; }
+		FConcertClientReplicationBridgeMock& GetBridgeMock() const { checkf(BridgeMock, TEXT("You forgot to call JoinReplication")); return *BridgeMock; }
+		IConcertClientReplicationManager& GetClientReplicationManager() const { checkf(ClientReplicationManager, TEXT("You forgot to call JoinReplication")); return *ClientReplicationManager; }
 
 	private:
 
+		/** Relevant for certain requests. Passed to ClientReplicationManager upon creation. */
+		const EConcertSyncSessionFlags SessionFlags;
+
 		/** Used to test "obvious" cases that should never fail in any test. */
 		FAutomationTestBase& TestContext;
 		/** The underlying client session */
diff --git a/Engine/Plugins/Developer/Concert/ConcertSync/ConcertSyncTest/Source/ConcertSyncTest/Private/Replication/Util/Spec/ReplicationServer.cpp b/Engine/Plugins/Developer/Concert/ConcertSync/ConcertSyncTest/Source/ConcertSyncTest/Private/Replication/Util/Spec/ReplicationServer.cpp
--- a/Engine/Plugins/Developer/Concert/ConcertSync/ConcertSyncTest/Source/ConcertSyncTest/Private/Replication/Util/Spec/ReplicationServer.cpp
+++ b/Engine/Plugins/Developer/Concert/ConcertSync/ConcertSyncTest/Source/ConcertSyncTest/Private/Replication/Util/Spec/ReplicationServer.cpp
@@ -7,16 +7,17 @@
 
 namespace UE::ConcertSyncTests::Replication
 {
-	FReplicationServer::FReplicationServer(FAutomationTestBase& TestContext)
-		: TestContext(TestContext)
+	FReplicationServer::FReplicationServer(FAutomationTestBase& TestContext, EConcertSyncSessionFlags InSessionFlags)
+		: SessionFlags(InSessionFlags)
+		, TestContext(TestContext)
 		, ServerSessionMock(MakeShared<FConcertServerSessionMock>())
-		, ServerReplicationManager(ConcertSyncServer::TestInterface::CreateServerReplicationManager(ServerSessionMock))
+		, ServerReplicationManager(ConcertSyncServer::TestInterface::CreateServerReplicationManager(ServerSessionMock, InSessionFlags))
 	{}
 
 	FReplicationClient& FReplicationServer::ConnectClient()
 	{
 		const FGuid ClientEndpointId(0, 0, 0, Clients.Num() + 1); // {0, 0, 0, 0} is used by the server.
-		Clients.Add(MakeUnique<FReplicationClient>(ClientEndpointId, *ServerSessionMock, TestContext));
+		Clients.Add(MakeUnique<FReplicationClient>(ClientEndpointId, SessionFlags, *ServerSessionMock, TestContext));
 		ServerSessionMock->ConnectClient(ClientEndpointId, *(Clients.Last()->GetClientSessionMock()));
 		return *Clients.Last();
 	}
diff --git a/Engine/Plugins/Developer/Concert/ConcertSync/ConcertSyncTest/Source/ConcertSyncTest/Private/Replication/Util/Spec/ReplicationServer.h b/Engine/Plugins/Developer/Concert/ConcertSync/ConcertSyncTest/Source/ConcertSyncTest/Private/Replication/Util/Spec/ReplicationServer.h
--- a/Engine/Plugins/Developer/Concert/ConcertSync/ConcertSyncTest/Source/ConcertSyncTest/Private/Replication/Util/Spec/ReplicationServer.h
+++ b/Engine/Plugins/Developer/Concert/ConcertSync/ConcertSyncTest/Source/ConcertSyncTest/Private/Replication/Util/Spec/ReplicationServer.h
@@ -2,8 +2,10 @@
 
 #pragma once
 
+#include "ConcertSyncSessionFlags.h"
 #include "Replication/IConcertServerReplicationManager.h"
 #include "Util/ClientServerCommunicationTest.h"
+
 #include "Templates/UnrealTemplate.h"
 
 class FAutomationTestBase;
@@ -17,7 +19,7 @@ namespace UE::ConcertSyncTests::Replication
 	{
 	public:
 
-		FReplicationServer(FAutomationTestBase& TestContext);
+		FReplicationServer(FAutomationTestBase& TestContext, EConcertSyncSessionFlags InSessionFlags = EConcertSyncSessionFlags::Default_MultiUserSession);
 
 		/** Connects a client to the server. */
 		FReplicationClient& ConnectClient();
@@ -29,6 +31,9 @@ namespace UE::ConcertSyncTests::Replication
 	
 	private:
 
+		/** Relevant to some requests. Passed to ServerReplicationManager and clients upon creation. */
+		const EConcertSyncSessionFlags SessionFlags;
+
 		/** Used to test "obvious" cases that should never fail in any test. */
 		FAutomationTestBase& TestContext;
 
diff --git a/Engine/Plugins/Developer/Concert/ConcertSync/ConcertSyncTest/Source/ConcertSyncTest/Private/Util/ClientServerCommunicationTest.h b/Engine/Plugins/Developer/Concert/ConcertSync/ConcertSyncTest/Source/ConcertSyncTest/Private/Util/ClientServerCommunicationTest.h
--- a/Engine/Plugins/Developer/Concert/ConcertSync/ConcertSyncTest/Source/ConcertSyncTest/Private/Util/ClientServerCommunicationTest.h
+++ b/Engine/Plugins/Developer/Concert/ConcertSync/ConcertSyncTest/Source/ConcertSyncTest/Private/Util/ClientServerCommunicationTest.h
@@ -287,10 +287,6 @@ namespace UE::ConcertSyncTests
 					// Dispatch the response.
 					ResponseHandler->HandleResponse(ResponsePayload.GetStructMemory());
 				}
-				else
-				{
-					check(false); // The test suite is not expected to fire any other result than Success or Failed.
-				}
 			}
 			else 
 			{
