Produce replication activities when client leaves replication

diff --git a/Engine/Plugins/Developer/Concert/ConcertSync/ConcertSyncClient/Source/ConcertSyncClient/Private/ConcertClientWorkspace.cpp b/Engine/Plugins/Developer/Concert/ConcertSync/ConcertSyncClient/Source/ConcertSyncClient/Private/ConcertClientWorkspace.cpp
--- a/Engine/Plugins/Developer/Concert/ConcertSync/ConcertSyncClient/Source/ConcertSyncClient/Private/ConcertClientWorkspace.cpp
+++ b/Engine/Plugins/Developer/Concert/ConcertSync/ConcertSyncClient/Source/ConcertSyncClient/Private/ConcertClientWorkspace.cpp
@@ -43,6 +43,7 @@
 #include "AssetRegistry/AssetRegistryModule.h"
 #include "Engine/World.h"
 #include "Engine/Engine.h"
+#include "Replication/Messages/ReplicationActivity.h"
 
 #if WITH_EDITOR
 	#include "Editor.h"
@@ -970,6 +971,7 @@ void FConcertClientWorkspace::HandleWorkspaceSyncActivityEvent(const FConcertSes
 	}
 
 	// Handle the activity correctly
+	static_assert(static_cast<uint8>(EConcertSyncActivityEventType::Count) == 6, "If you added an EConcertSyncActivityEventType entry, update this switch");
 	switch (Activity->EventType)
 	{
 	case EConcertSyncActivityEventType::Connection:
@@ -992,6 +994,11 @@ void FConcertClientWorkspace::HandleWorkspaceSyncActivityEvent(const FConcertSes
 		SetPackageActivity(*(FConcertSyncPackageActivity*)Activity);
 		break;
 
+	case EConcertSyncActivityEventType::Replication:
+		check(ActivityPayload.GetStruct()->IsChildOf(FConcertSyncReplicationActivity::StaticStruct()));
+		SetReplicationActivity(*(FConcertSyncReplicationActivity*)Activity);
+		break;
+
 	default:
 		checkf(false, TEXT("Unhandled EConcertSyncActivityEventType when syncing session activity"));
 		break;
@@ -1121,6 +1128,19 @@ void FConcertClientWorkspace::SetPackageActivity(const FConcertSyncPackageActivi
 	}
 }
 
+void FConcertClientWorkspace::SetReplicationActivity(const FConcertSyncReplicationActivity& InReplicationActivity)
+{
+	// Update this activity
+	if (LiveSession->GetSessionDatabase().SetReplicationActivity(InReplicationActivity))
+	{
+		PostActivityUpdated(InReplicationActivity);
+	}
+	else
+	{
+		UE_LOG(LogConcert, Error, TEXT("Failed to set replication activity '%s' on live session '%s': %s"), *LexToString(InReplicationActivity.ActivityId), *LiveSession->GetSession().GetName(), *LiveSession->GetSessionDatabase().GetLastError());
+	}
+}
+
 void FConcertClientWorkspace::PostActivityUpdated(const FConcertSyncActivity& InActivity)
 {
 	FConcertSyncActivity Activity;
diff --git a/Engine/Plugins/Developer/Concert/ConcertSync/ConcertSyncClient/Source/ConcertSyncClient/Private/ConcertClientWorkspace.h b/Engine/Plugins/Developer/Concert/ConcertSync/ConcertSyncClient/Source/ConcertSyncClient/Private/ConcertClientWorkspace.h
--- a/Engine/Plugins/Developer/Concert/ConcertSync/ConcertSyncClient/Source/ConcertSyncClient/Private/ConcertClientWorkspace.h
+++ b/Engine/Plugins/Developer/Concert/ConcertSync/ConcertSyncClient/Source/ConcertSyncClient/Private/ConcertClientWorkspace.h
@@ -36,6 +36,7 @@ namespace UE::ConcertSyncClient
 	};
 }
 
+struct FConcertSyncReplicationActivity;
 struct FScopedSlowTask;
 
 DECLARE_MULTICAST_DELEGATE(FOnWorkspaceEndFrameCompleted);
@@ -192,6 +193,14 @@ private:
 	 */
 	void SetPackageActivity(const FConcertSyncPackageActivity& InPackageActivity);
 
+	/**
+	 * Set a replication activity in the session database, creating or replacing it.
+	 * @note The endpoint ID referenced by the activity must exist in the database (@see SetEndpoint).
+	 *
+	 * @param InReplicationActivity		The replication activity to set.
+	 */
+	void SetReplicationActivity(const FConcertSyncReplicationActivity& InReplicationActivity);
+
 	/**
 	 * Called after any updated in the session database.
 	 *
diff --git a/Engine/Plugins/Developer/Concert/ConcertSync/ConcertSyncCore/Source/ConcertSyncCore/Private/Replication/Messages/ReplicationActivity.cpp b/Engine/Plugins/Developer/Concert/ConcertSync/ConcertSyncCore/Source/ConcertSyncCore/Private/Replication/Messages/ReplicationActivity.cpp
new file mode 100644
--- /dev/null
+++ b/Engine/Plugins/Developer/Concert/ConcertSync/ConcertSyncCore/Source/ConcertSyncCore/Private/Replication/Messages/ReplicationActivity.cpp
@@ -0,0 +1,52 @@
+﻿// Copyright Epic Games, Inc. All Rights Reserved.
+
+#include "Replication/Messages/ReplicationActivity.h"
+
+#include "Internationalization/Text.h"
+
+bool operator==(const FConcertSyncReplicationEvent& Left, const FConcertSyncReplicationEvent& Right)
+{
+	if (Left.ActivityType != Right.ActivityType)
+	{
+		return false;
+	}
+
+	static_assert(static_cast<int32>(EConcertSyncReplicationActivityType::Count) == 2, "If you added an EConcertSyncReplicationActivityType entry, update this switch");
+	switch (Left.ActivityType)
+	{
+	case EConcertSyncReplicationActivityType::None: return true;
+	case EConcertSyncReplicationActivityType::LeaveReplication:
+		{
+			FConcertSyncReplicationPayload_LeaveReplication LeftContent;
+			FConcertSyncReplicationPayload_LeaveReplication RightContent;
+			const bool bLeftSucceeded = Left.GetPayload(LeftContent);
+			const bool bRightSucceeded = Right.GetPayload(RightContent);
+			return bLeftSucceeded && bRightSucceeded && LeftContent == RightContent;
+		}
+	default: checkNoEntry(); return false;
+	}
+}
+
+FConcertSyncReplicationActivitySummary FConcertSyncReplicationActivitySummary::CreateSummaryForEvent(const FConcertSyncReplicationEvent& InEvent)
+{
+	FConcertSyncReplicationActivitySummary Summary;
+	Summary.ActivityType = InEvent.ActivityType;
+
+	static_assert(static_cast<uint8>(EConcertSyncReplicationActivityType::Count) == 2, "If you changed the enum entries, update this switch");
+	switch (Summary.ActivityType)
+	{
+	case EConcertSyncReplicationActivityType::LeaveReplication:
+		{
+			FConcertSyncReplicationPayload_LeaveReplication LeaveReplicationData;
+			InEvent.GetPayload(LeaveReplicationData);
+			
+			const FConcertSyncReplicationSummary_LeaveReplication Content { LeaveReplicationData.OwnedObjects };
+			Summary.Payload.SetTypedPayload(Content);
+		}
+		break;
+	case EConcertSyncReplicationActivityType::None: [[fallthrough]];
+	default: checkNoEntry(); break;
+	}
+	
+	return Summary;
+}
\ No newline at end of file
diff --git a/Engine/Plugins/Developer/Concert/ConcertSync/ConcertSyncCore/Source/ConcertSyncCore/Public/ConcertSyncSessionFlags.h b/Engine/Plugins/Developer/Concert/ConcertSync/ConcertSyncCore/Source/ConcertSyncCore/Public/ConcertSyncSessionFlags.h
--- a/Engine/Plugins/Developer/Concert/ConcertSync/ConcertSyncCore/Source/ConcertSyncCore/Public/ConcertSyncSessionFlags.h
+++ b/Engine/Plugins/Developer/Concert/ConcertSync/ConcertSyncCore/Source/ConcertSyncCore/Public/ConcertSyncSessionFlags.h
@@ -46,18 +46,21 @@ enum class EConcertSyncSessionFlags : uint32
 	ShouldSendPackageAutoSaves = 1<<12,
 	/** Use a sandbox to hold package writes during a session (requires EnablePackages) */
 	ShouldUsePackageSandbox = 1<<13,
+	
 	/**
 	 * For Concert Replication, clients are allowed to globally mute objects.
 	 * More formally, allows the use of FConcertReplication_ChangeMuteState_Requests.
 	 */
 	ShouldAllowGlobalMuting = 1<<14,
+	/** Whether EConcertSyncActivityEventType::Replication activities should be generated. */
+	ShouldEnableReplicationActivities = 1<<15,
 
 	/**
 	 * Defaults for different scenarios
 	 */
 	
 	/** Default settings for a multi-user session */
-	Default_MultiUserSession = EnableLiveSync | EnableConnectionHistory | EnableLocking | EnableTransactions | EnablePackages | EnableSequencer | EnablePresence | ShouldSendTransactionSnapshots | ShouldUsePackageSandbox | EnableReplication | ShouldAllowGlobalMuting,
+	Default_MultiUserSession = EnableLiveSync | EnableConnectionHistory | EnableLocking | EnableTransactions | EnablePackages | EnableSequencer | EnablePresence | ShouldSendTransactionSnapshots | ShouldUsePackageSandbox | EnableReplication | ShouldAllowGlobalMuting | ShouldEnableReplicationActivities,
 	/** Default settings for a disaster recovery session */
 	Default_DisasterRecoverySession = EnableTransactions | EnablePackages | EnableFileSharing | ShouldDiscardTransactionsOnPackageUnload | ShouldSendPackagePristineState | ShouldSendPackageAutoSaves,
 };
diff --git a/Engine/Plugins/Developer/Concert/ConcertSync/ConcertSyncCore/Source/ConcertSyncCore/Public/ConcertSyncSessionTypes.cpp b/Engine/Plugins/Developer/Concert/ConcertSync/ConcertSyncCore/Source/ConcertSyncCore/Public/ConcertSyncSessionTypes.cpp
--- a/Engine/Plugins/Developer/Concert/ConcertSync/ConcertSyncCore/Source/ConcertSyncCore/Public/ConcertSyncSessionTypes.cpp
+++ b/Engine/Plugins/Developer/Concert/ConcertSync/ConcertSyncCore/Source/ConcertSyncCore/Public/ConcertSyncSessionTypes.cpp
@@ -4,6 +4,7 @@
 #include "ConcertLogGlobal.h"
 
 #include "Misc/PackageName.h"
+#include "Replication/Messages/ReplicationActivity.h"
 
 #include UE_INLINE_GENERATED_CPP_BY_NAME(ConcertSyncSessionTypes)
 
@@ -211,7 +212,6 @@ FText FConcertSyncConnectionActivitySummary::CreateDisplayTextForUser(const FTex
 	return FText::Format(FormatPattern, Arguments);
 }
 
-
 FConcertSyncLockActivitySummary FConcertSyncLockActivitySummary::CreateSummaryForEvent(const FConcertSyncLockEvent& InEvent)
 {
 	const FString ResourceNameStr = InEvent.ResourceNames.Num() > 0 ? InEvent.ResourceNames[0].ToString() : FString();
@@ -830,5 +830,53 @@ FText FConcertSyncPackageActivitySummary::CreateDisplayTextForUser(const FText I
 	return FText::Format(FormatPattern, Arguments);
 }
 
+FText FConcertSyncReplicationActivitySummary::CreateDisplayText(const bool InUseRichText) const
+{
+	static_assert(static_cast<uint8>(EConcertSyncReplicationActivityType::Count) == 2, "If you added an EConcertSyncReplicationActivityType entry, update this switch");
+	switch (ActivityType)
+	{
+	case EConcertSyncReplicationActivityType::LeaveReplication:
+		{
+			FConcertSyncReplicationSummary_LeaveReplication Content;
+			const bool bRead = GetSummaryData(Content);
+			if (!bRead)
+			{
+				return LOCTEXT("CreateDisplayTextForUser_LeftActivity", "No data.");
+			}
+			
+			FFormatNamedArguments Arguments;
+			Arguments.Add(TEXT("NumObjects"), ActivitySummaryUtil::ToRichTextBold(FText::AsNumber(Content.OwnedObjects.Num()), InUseRichText));
+			return FText::Format(LOCTEXT("CreateDisplayTextForUser.LeftActivityFmt", "Stopped replicating {NumObjects} objects."), Arguments);
+			
+		}
+	case EConcertSyncReplicationActivityType::None: [[fallthrough]];
+		default: checkNoEntry(); return FConcertSyncActivitySummary::CreateDisplayText(InUseRichText);
+	}
+}
+
+FText FConcertSyncReplicationActivitySummary::CreateDisplayTextForUser(const FText InUserDisplayName, const bool InUseRichText) const
+{
+	static_assert(static_cast<uint8>(EConcertSyncReplicationActivityType::Count) == 2, "If you added an EConcertSyncReplicationActivityType entry, update this switch");
+	switch (ActivityType)
+	{
+	case EConcertSyncReplicationActivityType::LeaveReplication:
+		{
+			FConcertSyncReplicationSummary_LeaveReplication Content;
+			const bool bRead = GetSummaryData(Content);
+			if (!bRead)
+			{
+				return LOCTEXT("CreateDisplayTextForUser_LeftActivity", "No data.");
+			}
+			
+			FFormatNamedArguments Arguments;
+			Arguments.Add(TEXT("UserName"), ActivitySummaryUtil::ToRichTextBold(InUserDisplayName, InUseRichText));
+			Arguments.Add(TEXT("NumObjects"), ActivitySummaryUtil::ToRichTextBold(FText::AsNumber(Content.OwnedObjects.Num()), InUseRichText));
+			return FText::Format(LOCTEXT("CreateDisplayTextForUser.LeftActivityFmt", "{UserName} stopped replicating {NumObjects} objects."), Arguments);
+		}
+	case EConcertSyncReplicationActivityType::None: [[fallthrough]];
+	default: checkNoEntry(); return FConcertSyncActivitySummary::CreateDisplayText(InUseRichText);
+	}
+}
+
 #undef LOCTEXT_NAMESPACE
 
diff --git a/Engine/Plugins/Developer/Concert/ConcertSync/ConcertSyncCore/Source/ConcertSyncCore/Public/Replication/Messages/ReplicationActivity.h b/Engine/Plugins/Developer/Concert/ConcertSync/ConcertSyncCore/Source/ConcertSyncCore/Public/Replication/Messages/ReplicationActivity.h
--- a/Engine/Plugins/Developer/Concert/ConcertSync/ConcertSyncCore/Source/ConcertSyncCore/Public/Replication/Messages/ReplicationActivity.h
+++ b/Engine/Plugins/Developer/Concert/ConcertSync/ConcertSyncCore/Source/ConcertSyncCore/Public/Replication/Messages/ReplicationActivity.h
@@ -9,6 +9,8 @@
 #include "Replication/Data/ReplicationStream.h"
 #include "ReplicationActivity.generated.h"
 
+class FText;
+
 /** Identifies the FConcertSyncReplicationEvent::Payload struct type */
 UENUM()
 enum class EConcertSyncReplicationActivityType : uint8
@@ -43,6 +45,16 @@ struct FConcertSyncReplicationPayload_LeaveReplication
 	friend bool operator!=(const FConcertSyncReplicationPayload_LeaveReplication&, const FConcertSyncReplicationPayload_LeaveReplication&) = default;
 };
 
+USTRUCT()
+struct FConcertSyncReplicationSummary_LeaveReplication
+{
+	GENERATED_BODY()
+
+	/** The objects the client had authority over when they left. */
+	UPROPERTY()
+	TArray<FConcertObjectInStreamID> OwnedObjects;
+};
+
 namespace UE::ConcertSyncCore
 {
 	inline FString GetReplicationActivityPayloadTypePathName(EConcertSyncReplicationActivityType Type)
@@ -94,28 +106,7 @@ struct FConcertSyncReplicationEvent
 		return Payload.GetTypedPayload(Result);
 	}
 
-	friend bool operator==(const FConcertSyncReplicationEvent& Left, const FConcertSyncReplicationEvent& Right)
-	{
-		if (Left.ActivityType != Right.ActivityType)
-		{
-			return false;
-		}
-
-		static_assert(static_cast<int32>(EConcertSyncReplicationActivityType::Count) == 2, "If you added an EConcertSyncReplicationActivityType entry, update this switch");
-		switch (Left.ActivityType)
-		{
-		case EConcertSyncReplicationActivityType::None: return true;
-		case EConcertSyncReplicationActivityType::LeaveReplication:
-		{
-			FConcertSyncReplicationPayload_LeaveReplication LeftContent;
-			FConcertSyncReplicationPayload_LeaveReplication RightContent;
-			const bool bLeftSucceeded = Left.GetPayload(LeftContent);
-			const bool bRightSucceeded = Right.GetPayload(RightContent);
-			return bLeftSucceeded && bRightSucceeded && LeftContent == RightContent;
-		}
-		default: checkNoEntry(); return false;
-		}
-	}
+	CONCERTSYNCCORE_API friend bool operator==(const FConcertSyncReplicationEvent& Left, const FConcertSyncReplicationEvent& Right);
 	friend bool operator!=(const FConcertSyncReplicationEvent&, const FConcertSyncReplicationEvent&) = default;
 };
 
@@ -139,4 +130,35 @@ struct FConcertSyncReplicationActivity : public FConcertSyncActivity
 	/** The transaction event data associated with this activity */
 	UPROPERTY()
 	FConcertSyncReplicationEvent EventData;
+};
+
+/** Summary for a lock activity entry in a Concert Sync Session */
+USTRUCT()
+struct CONCERTSYNCCORE_API FConcertSyncReplicationActivitySummary : public FConcertSyncActivitySummary
+{
+	GENERATED_BODY()
+
+	/** The type of replication event we summarize */
+	UPROPERTY()
+	EConcertSyncReplicationActivityType ActivityType;
+
+	/** The summary data. The underlying type depends on ActivityType. */
+	UPROPERTY()
+	FConcertSessionSerializedPayload Payload{ EConcertPayloadSerializationMethod::Cbor };
+
+	bool GetSummaryData(FConcertSyncReplicationSummary_LeaveReplication& Data) const
+	{
+		check(ActivityType == EConcertSyncReplicationActivityType::LeaveReplication);
+		return Payload.GetTypedPayload(Data);
+	}
+
+	/** Create this summary from a replication event */
+	static FConcertSyncReplicationActivitySummary CreateSummaryForEvent(const FConcertSyncReplicationEvent& InEvent);
+
+protected:
+	
+	//~ Begin FConcertSyncActivitySummary Interface
+	virtual FText CreateDisplayText(const bool InUseRichText) const override;
+	virtual FText CreateDisplayTextForUser(const FText InUserDisplayName, const bool InUseRichText) const override;
+	//~ End FConcertSyncActivitySummary Interface
 };
\ No newline at end of file
diff --git a/Engine/Plugins/Developer/Concert/ConcertSync/ConcertSyncServer/Source/ConcertSyncServer/Private/ConcertServerWorkspace.cpp b/Engine/Plugins/Developer/Concert/ConcertSync/ConcertSyncServer/Source/ConcertSyncServer/Private/ConcertServerWorkspace.cpp
--- a/Engine/Plugins/Developer/Concert/ConcertSync/ConcertSyncServer/Source/ConcertSyncServer/Private/ConcertServerWorkspace.cpp
+++ b/Engine/Plugins/Developer/Concert/ConcertSync/ConcertSyncServer/Source/ConcertSyncServer/Private/ConcertServerWorkspace.cpp
@@ -30,6 +30,19 @@ FConcertServerWorkspace::~FConcertServerWorkspace()
 	UnbindSession();
 }
 
+void FConcertServerWorkspace::ProduceClientLeaveReplicationActivity(const FGuid& EndpointId, const FConcertSyncReplicationPayload_LeaveReplication& EventData)
+{
+	if (EnumHasAnyFlags(LiveSession->GetSessionFlags(), EConcertSyncSessionFlags::ShouldEnableReplicationActivities))
+	{
+		FConcertSyncReplicationActivity Activity;
+		Activity.EndpointId = EndpointId;
+		Activity.EventData.SetPayload(EventData);
+		Activity.EventSummary.SetTypedPayload(FConcertSyncReplicationActivitySummary::CreateSummaryForEvent(Activity.EventData));
+		Activity.bIgnored = ShouldIgnoreClientActivityOnRestore(EndpointId);
+		AddReplicationActivity(Activity);
+	}
+}
+
 void FConcertServerWorkspace::BindSession(const TSharedRef<FConcertSyncServerLiveSession>& InLiveSession)
 {
 	check(InLiveSession->IsValidSession());
@@ -282,6 +295,7 @@ void FConcertServerWorkspace::HandleSyncRequestedEvent(const FConcertSessionCont
 		
 		SyncCommandQueue->QueueCommand(Context.SourceEndpointId, [this, SyncActivityId = InActivityId, SyncEventType = InEventType](const FConcertServerSyncCommandQueue::FSyncCommandContext& InSyncCommandContext, const FGuid& InEndpointId)
 		{
+			static_assert(static_cast<uint8>(EConcertSyncActivityEventType::Count) == 6, "If you added an EConcertSyncActivityEventType entry, update this switch");
 			switch (SyncEventType)
 			{
 			case EConcertSyncActivityEventType::Connection:
@@ -306,6 +320,9 @@ void FConcertServerWorkspace::HandleSyncRequestedEvent(const FConcertSessionCont
 				}
 				break;
 			}
+			case EConcertSyncActivityEventType::Replication:
+				SendSyncReplicationActivityEvent(InEndpointId, SyncActivityId, InSyncCommandContext.GetNumRemainingCommands());
+				break;
 			default:
 				checkf(false, TEXT("Unhandled EConcertSyncActivityEventType when syncing session activity"));
 				break;
diff --git a/Engine/Plugins/Developer/Concert/ConcertSync/ConcertSyncServer/Source/ConcertSyncServer/Private/ConcertServerWorkspace.h b/Engine/Plugins/Developer/Concert/ConcertSync/ConcertSyncServer/Source/ConcertSyncServer/Private/ConcertServerWorkspace.h
--- a/Engine/Plugins/Developer/Concert/ConcertSync/ConcertSyncServer/Source/ConcertSyncServer/Private/ConcertServerWorkspace.h
+++ b/Engine/Plugins/Developer/Concert/ConcertSync/ConcertSyncServer/Source/ConcertSyncServer/Private/ConcertServerWorkspace.h
@@ -3,9 +3,10 @@
 #pragma once
 
 #include "CoreMinimal.h"
-#include "IConcertSessionHandler.h"
 #include "ConcertWorkspaceMessages.h"
+#include "IConcertSessionHandler.h"
 #include "Misc/Optional.h"
+#include "Replication/IReplicationWorkspace.h"
 
 class IConcertFileSharingService;
 class IConcertMessage;
@@ -17,6 +18,7 @@ class FConcertServerDataStore;
 struct FConcertMessageContext;
 struct FConcertEndpointContext;
 struct FConcertSyncReplicationActivity;
+struct FConcertSyncReplicationPayload_LeaveReplication;
 struct FConcertTransactionSnapshotEvent;
 struct FConcertTransactionFinalizedEvent;
 
@@ -29,11 +31,16 @@ enum class EConcertLockFlags : uint8
 };
 ENUM_CLASS_FLAGS(EConcertLockFlags);
 
-class FConcertServerWorkspace
+class FConcertServerWorkspace : public UE::ConcertSyncServer::Replication::IReplicationWorkspace
 {
 public:
+	
 	explicit FConcertServerWorkspace(const TSharedRef<FConcertSyncServerLiveSession>& InLiveSession, TSharedPtr<IConcertFileSharingService> InFileSharingService);
-	~FConcertServerWorkspace();
+	virtual ~FConcertServerWorkspace() override;
+
+	//~ Begin IReplicationWorkspace Interface
+	virtual void ProduceClientLeaveReplicationActivity(const FGuid& EndpointId, const FConcertSyncReplicationPayload_LeaveReplication& EventData) override;
+	//~ End IReplicationWorkspace Interface
 
 private:
 	
diff --git a/Engine/Plugins/Developer/Concert/ConcertSync/ConcertSyncServer/Source/ConcertSyncServer/Private/ConcertSyncServer.cpp b/Engine/Plugins/Developer/Concert/ConcertSync/ConcertSyncServer/Source/ConcertSyncServer/Private/ConcertSyncServer.cpp
--- a/Engine/Plugins/Developer/Concert/ConcertSync/ConcertSyncServer/Source/ConcertSyncServer/Private/ConcertSyncServer.cpp
+++ b/Engine/Plugins/Developer/Concert/ConcertSync/ConcertSyncServer/Source/ConcertSyncServer/Private/ConcertSyncServer.cpp
@@ -521,11 +521,14 @@ void FConcertSyncServer::OnArchivedSessionRenamedImpl(const IConcertServer& InSe
 	ConcertSyncServerUtils::WriteSessionInfoToDirectory(InArchivedSessionRoot, InArchivedSessionInfo);
 }
 
-void FConcertSyncServer::CreateWorkspace(const TSharedRef<FConcertSyncServerLiveSession>& InLiveSession)
+TSharedRef<FConcertServerWorkspace> FConcertSyncServer::CreateWorkspace(const TSharedRef<FConcertSyncServerLiveSession>& InLiveSession)
 {
 	check(InLiveSession->IsValidSession());
 	DestroyWorkspace(InLiveSession);
-	LiveSessionWorkspaces.Add(InLiveSession->GetSession().GetId(), MakeShared<FConcertServerWorkspace>(InLiveSession, FileSharingService));
+	
+	const TSharedRef<FConcertServerWorkspace> Workspace = MakeShared<FConcertServerWorkspace>(InLiveSession, FileSharingService);
+	LiveSessionWorkspaces.Add(InLiveSession->GetSession().GetId(), Workspace);
+	return Workspace;
 }
 
 void FConcertSyncServer::DestroyWorkspace(const TSharedRef<FConcertSyncServerLiveSession>& InLiveSession)
@@ -545,13 +548,13 @@ void FConcertSyncServer::DestroySequencerManager(const TSharedRef<FConcertSyncSe
 	LiveSessionSequencerManagers.Remove(InLiveSession->GetSession().GetId());
 }
 
-void FConcertSyncServer::CreateReplicationManager(const TSharedRef<IConcertServerSession>& InSession, EConcertSyncSessionFlags InSessionFlags)
+void FConcertSyncServer::CreateReplicationManager(const TSharedRef<IConcertServerSession>& InSession, UE::ConcertSyncServer::Replication::IReplicationWorkspace& InWorkspace, EConcertSyncSessionFlags InSessionFlags)
 {
 	const FGuid& SessionId = InSession->GetId(); 
 	DestroyReplicationManager(SessionId);
 	LiveSessionReplicationManagers.Add(
 		SessionId,
-		MakeShared<UE::ConcertSyncServer::Replication::FConcertServerReplicationManager>(InSession, InSessionFlags)
+		MakeShared<UE::ConcertSyncServer::Replication::FConcertServerReplicationManager>(InSession, InWorkspace, InSessionFlags)
 		);
 }
 
@@ -568,7 +571,7 @@ bool FConcertSyncServer::CreateLiveSession(const TSharedRef<IConcertServerSessio
 	if (LiveSession->IsValidSession())
 	{
 		LiveSessions.Add(InSession->GetId(), LiveSession);
-		CreateWorkspace(LiveSession.ToSharedRef());
+		const TSharedRef<FConcertServerWorkspace> Workspace = CreateWorkspace(LiveSession.ToSharedRef());
 		if (EnumHasAnyFlags(SessionFlags, EConcertSyncSessionFlags::EnableSequencer))
 		{
 			CreateSequencerManager(LiveSession.ToSharedRef());
@@ -577,7 +580,7 @@ bool FConcertSyncServer::CreateLiveSession(const TSharedRef<IConcertServerSessio
 		// Create Replication Manager
 		if (EnumHasAnyFlags(LiveSession->GetSessionFlags(), EConcertSyncSessionFlags::EnableReplication))
 		{
-			CreateReplicationManager(InSession, LiveSession->GetSessionFlags());
+			CreateReplicationManager(InSession, *Workspace, LiveSession->GetSessionFlags());
 		}
 
 		// We needn't call OnActivityProduced().Remove(...) because the subscription needs to stay for the lifetime of FConcertSyncServerLiveSession::SessionDatabase
diff --git a/Engine/Plugins/Developer/Concert/ConcertSync/ConcertSyncServer/Source/ConcertSyncServer/Private/ConcertSyncServer.h b/Engine/Plugins/Developer/Concert/ConcertSync/ConcertSyncServer/Source/ConcertSyncServer/Private/ConcertSyncServer.h
--- a/Engine/Plugins/Developer/Concert/ConcertSync/ConcertSyncServer/Source/ConcertSyncServer/Private/ConcertSyncServer.h
+++ b/Engine/Plugins/Developer/Concert/ConcertSync/ConcertSyncServer/Source/ConcertSyncServer/Private/ConcertSyncServer.h
@@ -17,6 +17,7 @@ class FConcertSyncSessionDatabase;
 namespace UE::ConcertSyncServer::Replication
 {
 	class FConcertServerReplicationManager;
+	class IReplicationWorkspace;
 }
 
 struct FConcertSessionFilter;
@@ -55,13 +56,14 @@ public:
 	void OnArchivedSessionRenamedImpl(const IConcertServer& InServer, const FString& InArchivedSessionRoot, const FConcertSessionInfo& InArchivedSessionInfo);
 
 private:
-	void CreateWorkspace(const TSharedRef<FConcertSyncServerLiveSession>& InLiveSession);
+	
+	TSharedRef<FConcertServerWorkspace> CreateWorkspace(const TSharedRef<FConcertSyncServerLiveSession>& InLiveSession);
 	void DestroyWorkspace(const TSharedRef<FConcertSyncServerLiveSession>& InLiveSession);
 
 	void CreateSequencerManager(const TSharedRef<FConcertSyncServerLiveSession>& InLiveSession);
 	void DestroySequencerManager(const TSharedRef<FConcertSyncServerLiveSession>& InLiveSession);
 	
-	void CreateReplicationManager(const TSharedRef<IConcertServerSession>& InSession, EConcertSyncSessionFlags InSessionFlags);
+	void CreateReplicationManager(const TSharedRef<IConcertServerSession>& InSession, UE::ConcertSyncServer::Replication::IReplicationWorkspace& InWorkspace, EConcertSyncSessionFlags InSessionFlags);
 	void DestroyReplicationManager(const FGuid& SessionId); 
 
 	bool CreateLiveSession(const TSharedRef<IConcertServerSession>& InSession, const FInternalLiveSessionCreationParams& AdditionalParams);
diff --git a/Engine/Plugins/Developer/Concert/ConcertSync/ConcertSyncServer/Source/ConcertSyncServer/Private/Replication/AuthorityManager.cpp b/Engine/Plugins/Developer/Concert/ConcertSync/ConcertSyncServer/Source/ConcertSyncServer/Private/Replication/AuthorityManager.cpp
--- a/Engine/Plugins/Developer/Concert/ConcertSync/ConcertSyncServer/Source/ConcertSyncServer/Private/Replication/AuthorityManager.cpp
+++ b/Engine/Plugins/Developer/Concert/ConcertSync/ConcertSyncServer/Source/ConcertSyncServer/Private/Replication/AuthorityManager.cpp
@@ -98,6 +98,26 @@ namespace UE::ConcertSyncServer::Replication
 		}
 	}
 
+	TArray<FConcertObjectInStreamID> FAuthorityManager::GetOwnedObjects(const FClientId& ClientId) const
+	{
+		TArray<FConcertObjectInStreamID> Result;
+		const FClientAuthorityData* AuthorityData = ClientAuthorityData.Find(ClientId);
+		if (!AuthorityData)
+		{
+			return Result;
+		}
+
+		for (const TPair<FStreamId, TSet<FSoftObjectPath>>& Pair : AuthorityData->OwnedObjects)
+		{
+			for (const FSoftObjectPath& ObjectPath : Pair.Value)
+			{
+				Result.Add({ Pair.Key, ObjectPath });
+			}
+		}
+		
+		return Result;
+	}
+
 	FAuthorityManager::EAuthorityResult FAuthorityManager::EnumerateAuthorityConflicts(
 		const FConcertReplicatedObjectId& Object, 
 		const FConcertPropertySelection* OverwriteProperties,
diff --git a/Engine/Plugins/Developer/Concert/ConcertSync/ConcertSyncServer/Source/ConcertSyncServer/Private/Replication/AuthorityManager.h b/Engine/Plugins/Developer/Concert/ConcertSync/ConcertSyncServer/Source/ConcertSyncServer/Private/Replication/AuthorityManager.h
--- a/Engine/Plugins/Developer/Concert/ConcertSync/ConcertSyncServer/Source/ConcertSyncServer/Private/Replication/AuthorityManager.h
+++ b/Engine/Plugins/Developer/Concert/ConcertSync/ConcertSyncServer/Source/ConcertSyncServer/Private/Replication/AuthorityManager.h
@@ -33,7 +33,6 @@ namespace UE::ConcertSyncServer::Replication
 		using FClientId = FGuid;
 		using FProcessAuthorityConflict = TFunctionRef<EBreakBehavior(const FClientId& ClientId, const FStreamId& StreamId, const FConcertPropertyChain& WrittenProperties)>;
 		
-
 		/**
 		 * @param InGetters Gets information about clients' registered streams 
 		 * @param InSession The session to handle authority requests on
@@ -50,6 +49,8 @@ namespace UE::ConcertSyncServer::Replication
 
 		/** Utility for iterating authority a client has for a given stream. */
 		void EnumerateAuthority(const FClientId& ClientId, const FStreamId& StreamId, TFunctionRef<EBreakBehavior(const FSoftObjectPath& Object)> Callback) const;
+		/** Util for getting all client owned objects as array. */
+		TArray<FConcertObjectInStreamID> GetOwnedObjects(const FClientId& ClientId) const;
 
 		enum class EAuthorityResult : uint8
 		{
diff --git a/Engine/Plugins/Developer/Concert/ConcertSync/ConcertSyncServer/Source/ConcertSyncServer/Private/Replication/ConcertServerReplicationManager.cpp b/Engine/Plugins/Developer/Concert/ConcertSync/ConcertSyncServer/Source/ConcertSyncServer/Private/Replication/ConcertServerReplicationManager.cpp
--- a/Engine/Plugins/Developer/Concert/ConcertSync/ConcertSyncServer/Source/ConcertSyncServer/Private/Replication/ConcertServerReplicationManager.cpp
+++ b/Engine/Plugins/Developer/Concert/ConcertSync/ConcertSyncServer/Source/ConcertSyncServer/Private/Replication/ConcertServerReplicationManager.cpp
@@ -4,13 +4,16 @@
 
 #include "AuthorityManager.h"
 #include "ConcertLogGlobal.h"
+#include "ConcertServerWorkspace.h"
 #include "IConcertSession.h"
 #include "Replication/ConcertReplicationClient.h"
 #include "Replication/Data/ClientQueriedInfo.h"
 #include "Replication/Formats/FullObjectFormat.h"
+#include "Replication/IReplicationWorkspace.h"
 #include "Replication/Messages/ChangeStream.h"
 #include "Replication/Messages/ClientQuery.h"
 #include "Replication/Messages/Handshake.h"
+#include "Replication/Messages/ReplicationActivity.h"
 #include "Replication/Processing/ObjectReplicationCache.h"
 #include "Util/JoinRequestValidation.h"
 #include "Util/LogUtils.h"
@@ -25,9 +28,12 @@ namespace UE::ConcertSyncServer::Replication
 	
 	FConcertServerReplicationManager::FConcertServerReplicationManager(
 		TSharedRef<IConcertServerSession> InLiveSession,
-		EConcertSyncSessionFlags SessionFlags
+		IReplicationWorkspace& InServerWorkspace,
+		EConcertSyncSessionFlags InSessionFlags
 		)
-		: Session(MoveTemp(InLiveSession))
+		: Session(InLiveSession)
+		, ServerWorkspace(InServerWorkspace)
+		, SessionFlags(InSessionFlags)
 		, ReplicationFormat(MakeUnique<ConcertSyncCore::FFullObjectFormat>())
 		, AuthorityManager(*this, Session)
 		, MuteManager(*Session, ServerObjectCache, SessionFlags)
@@ -262,6 +268,8 @@ namespace UE::ConcertSyncServer::Replication
 		const bool bRemoved = Clients.RemoveAndCopyValue(EndpointId, RemovedClient);
 		check(bRemoved);
 
+		ProduceClientLeftActivity(*RemovedClient);
+
 		// ServerObjectCache should be updated before anyone else that may rely on its state.
 		ServerObjectCache.OnPostClientLeft(EndpointId, RemovedClient->GetStreamDescriptions());
 		
@@ -273,6 +281,19 @@ namespace UE::ConcertSyncServer::Replication
 		SyncControlManager.OnPostClientLeft(EndpointId);
 	}
 
+	void FConcertServerReplicationManager::ProduceClientLeftActivity(const FConcertReplicationClient& Client) const
+	{
+		if (EnumHasAnyFlags(SessionFlags, EConcertSyncSessionFlags::ShouldEnableReplicationActivities))
+		{
+			const FGuid& EndpointId = Client.GetClientEndpointId();
+			
+			FConcertSyncReplicationPayload_LeaveReplication LeaveReplication;
+			LeaveReplication.Streams = Client.GetStreamDescriptions();
+			LeaveReplication.OwnedObjects = AuthorityManager.GetOwnedObjects(EndpointId);
+			ServerWorkspace.ProduceClientLeaveReplicationActivity(EndpointId, LeaveReplication);
+		}
+	}
+
 	void FConcertServerReplicationManager::Tick(IConcertServerSession& InSession, float InDeltaTime)
 	{
 		// TODO UE-190714: We may want to set a time budget for clients.
diff --git a/Engine/Plugins/Developer/Concert/ConcertSync/ConcertSyncServer/Source/ConcertSyncServer/Private/Replication/ConcertServerReplicationManager.h b/Engine/Plugins/Developer/Concert/ConcertSync/ConcertSyncServer/Source/ConcertSyncServer/Private/Replication/ConcertServerReplicationManager.h
--- a/Engine/Plugins/Developer/Concert/ConcertSync/ConcertSyncServer/Source/ConcertSyncServer/Private/Replication/ConcertServerReplicationManager.h
+++ b/Engine/Plugins/Developer/Concert/ConcertSync/ConcertSyncServer/Source/ConcertSyncServer/Private/Replication/ConcertServerReplicationManager.h
@@ -15,10 +15,12 @@
 #include "Replication/Processing/ServerObjectReplicationReceiver.h"
 #include "SyncControlManager.h"
 
+#include "HAL/Platform.h"
 #include "Templates/SharedPointer.h"
 #include "Templates/Tuple.h"
 #include "Templates/UnrealTemplate.h"
 
+class FConcertServerWorkspace;
 class IConcertClientReplicationBridge;
 class IConcertServerSession;
 
@@ -39,6 +41,7 @@ namespace UE::ConcertSyncCore
 namespace UE::ConcertSyncServer::Replication
 {
 	class FAuthorityManager;
+	class IReplicationWorkspace;
 
 	/**
 	 * Manages all server-side systems relevant to the Replication features.
@@ -53,7 +56,7 @@ namespace UE::ConcertSyncServer::Replication
 	{
 	public:
 
-		explicit FConcertServerReplicationManager(TSharedRef<IConcertServerSession> InLiveSession, EConcertSyncSessionFlags SessionFlags);
+		explicit FConcertServerReplicationManager(TSharedRef<IConcertServerSession> InLiveSession, IReplicationWorkspace& InServerWorkspace UE_LIFETIMEBOUND, EConcertSyncSessionFlags InSessionFlags);
 		virtual ~FConcertServerReplicationManager() override;
 
 		const FAuthorityManager& GetAuthorityManager() const { return AuthorityManager; }
@@ -69,7 +72,11 @@ namespace UE::ConcertSyncServer::Replication
 	private:
 		
 		/** Session instance this manager was created for. */
-		TSharedRef<IConcertServerSession> Session;
+		const TSharedRef<IConcertServerSession> Session;
+		/** Used to produce replication activities. */
+		IReplicationWorkspace& ServerWorkspace;
+		/** Used to determine which dynamic features are enabled. */
+		const EConcertSyncSessionFlags SessionFlags;
 		
 		/** Responsible for analysing received replication data. */
 		TUniquePtr<ConcertSyncCore::IObjectReplicationFormat> ReplicationFormat;
@@ -118,7 +125,10 @@ namespace UE::ConcertSyncServer::Replication
 		void HandleLeaveReplicationSessionRequest(const FConcertSessionContext& ConcertSessionContext, const FConcertReplication_LeaveEvent& EventData);
 		void OnConnectionChanged(IConcertServerSession& ConcertServerSession, EConcertClientStatus ConcertClientStatus, const FConcertSessionClientInfo& ConcertSessionClientInfo);
 
+		/** Cleans up the client's replication state after leaving. */
 		void OnClientLeftReplication(const FGuid& EndpointId);
+		/** Calls FConcertServerWorkspace::AddReplicationActivity with the current client state so it can be restored upon re-joining. */
+		void ProduceClientLeftActivity(const FConcertReplicationClient& Client) const;
 		
 		/**
 		 * Ticks all clients which causes clients to process pending data and send it to the corresponding endpoints.
diff --git a/Engine/Plugins/Developer/Concert/ConcertSync/ConcertSyncServer/Source/ConcertSyncServer/Private/Replication/ReplicationTestInterface.cpp b/Engine/Plugins/Developer/Concert/ConcertSync/ConcertSyncServer/Source/ConcertSyncServer/Private/Replication/ReplicationTestInterface.cpp
--- a/Engine/Plugins/Developer/Concert/ConcertSync/ConcertSyncServer/Source/ConcertSyncServer/Private/Replication/ReplicationTestInterface.cpp
+++ b/Engine/Plugins/Developer/Concert/ConcertSync/ConcertSyncServer/Source/ConcertSyncServer/Private/Replication/ReplicationTestInterface.cpp
@@ -15,9 +15,10 @@ namespace UE::ConcertSyncServer::TestInterface
 {
 	CONCERTSYNCSERVER_API TSharedRef<Replication::IConcertServerReplicationManager> CreateServerReplicationManager(
 		TSharedRef<IConcertServerSession> InLiveSession,
+		Replication::IReplicationWorkspace& InWorkspace,
 		EConcertSyncSessionFlags InSessionFlags
 		)
 	{
-		return MakeShared<Replication::FConcertServerReplicationManager>(MoveTemp(InLiveSession), InSessionFlags);
+		return MakeShared<Replication::FConcertServerReplicationManager>(InLiveSession, InWorkspace, InSessionFlags);
 	}
 }
diff --git a/Engine/Plugins/Developer/Concert/ConcertSync/ConcertSyncServer/Source/ConcertSyncServer/Public/Replication/IReplicationWorkspace.h b/Engine/Plugins/Developer/Concert/ConcertSync/ConcertSyncServer/Source/ConcertSyncServer/Public/Replication/IReplicationWorkspace.h
new file mode 100644
--- /dev/null
+++ b/Engine/Plugins/Developer/Concert/ConcertSync/ConcertSyncServer/Source/ConcertSyncServer/Public/Replication/IReplicationWorkspace.h
@@ -0,0 +1,23 @@
+﻿// Copyright Epic Games, Inc. All Rights Reserved.
+
+#pragma once
+
+struct FConcertSyncReplicationPayload_LeaveReplication;
+struct FGuid;
+
+namespace UE::ConcertSyncServer::Replication
+{
+	/**
+	 * Interface that FConcertServerReplicationManager uses to interact with the FConcertServerWorkspace.
+	 * Allows mocking in unit tests, which is the only reason it's in the public module interface.
+	 */
+	class IReplicationWorkspace
+	{
+	public:
+		
+		/** Creates a replication activity for the client leaving replication if the session has the EConcertSyncSessionFlags::ShouldEnableReplicationActivities flag. */
+		virtual void ProduceClientLeaveReplicationActivity(const FGuid& EndpointId, const FConcertSyncReplicationPayload_LeaveReplication& EventData) = 0;
+
+		virtual ~IReplicationWorkspace() = default;
+	};
+}
diff --git a/Engine/Plugins/Developer/Concert/ConcertSync/ConcertSyncTest/Source/ConcertSyncTest/Private/Replication/ActivityProduction.spec.cpp b/Engine/Plugins/Developer/Concert/ConcertSync/ConcertSyncTest/Source/ConcertSyncTest/Private/Replication/ActivityProduction.spec.cpp
new file mode 100644
--- /dev/null
+++ b/Engine/Plugins/Developer/Concert/ConcertSync/ConcertSyncTest/Source/ConcertSyncTest/Private/Replication/ActivityProduction.spec.cpp
@@ -0,0 +1,123 @@
+﻿// Copyright Epic Games, Inc. All Rights Reserved.
+
+#include "Util/Spec/ReplicationClient.h"
+#include "Util/Spec/ReplicationServer.h"
+
+#include "Misc/AutomationTest.h"
+#include "Replication/Messages/ReplicationActivity.h"
+#include "Util/Mocks/ReplicationWorkspaceCallDetectorMock.h"
+#include "Util/Spec/ObjectTestReplicator.h"
+
+namespace UE::ConcertSyncTests::Replication
+{
+	BEGIN_DEFINE_SPEC(FLeaveReplicationActivitySpec, "Editor.Concert.Replication.Activity", EAutomationTestFlags::EditorContext | EAutomationTestFlags::EngineFilter)
+		/** Detects calls into the workspace */
+		TSharedPtr<FReplicationWorkspaceCallDetectorMock> WorkspaceMock;
+	
+		TUniquePtr<FObjectTestReplicator> ObjectReplicator;
+		TUniquePtr<FReplicationServer> Server;
+		FReplicationClient* Client = nullptr;
+	
+		FGuid SenderStreamId = FGuid::NewGuid();
+
+		void ValidateLeaveActivity()
+		{
+			if (WorkspaceMock->LastCall_ProduceClientLeaveReplicationActivity)
+			{
+				const auto[EndpointId, EventData] = *WorkspaceMock->LastCall_ProduceClientLeaveReplicationActivity;
+				TestEqual(TEXT("EndpointId"), EndpointId, Client->GetEndpointId());
+				
+				TestEqual(TEXT("Streams.Num() == 1"), EventData.Streams.Num(), 1);
+				if (EventData.Streams.Num() == 1)
+				{
+					const FConcertBaseStreamInfo& StreamInfo = EventData.Streams[0].BaseDescription;
+					TestEqual(TEXT("StreamId"), StreamInfo.Identifier, SenderStreamId);
+					
+					const FConcertReplicatedObjectInfo* ObjectInfo = StreamInfo.ReplicationMap.ReplicatedObjects.Find(ObjectReplicator->TestObject);
+					if (ObjectInfo)
+					{
+						TestEqual(TEXT("ReplicatedObjects.Num() == 1"),StreamInfo.ReplicationMap.ReplicatedObjects.Num(), 1);
+						TestTrue(TEXT("Class correct"), ObjectInfo->ClassPath == ObjectReplicator->TestObject->GetClass());
+						
+						const ConcertSyncClient::Replication::FJoinReplicatedSessionArgs SenderArgs = ObjectReplicator->CreateSenderArgs(SenderStreamId);
+						const FConcertPropertySelection& ExpectedProperties = SenderArgs.Streams[0].BaseDescription.ReplicationMap.ReplicatedObjects[ObjectReplicator->TestObject].PropertySelection;
+						TestEqual(TEXT("Properties equal"), ObjectInfo->PropertySelection, ExpectedProperties);
+					}
+					else
+					{
+						AddError(TEXT("Replicated object not listed in stream"));
+					}
+				}
+				
+				TestEqual(TEXT("OwnedObjects.Num() == 1"), EventData.OwnedObjects.Num(), 1);
+				TestTrue(TEXT("OwnedObjects contains owned object"), EventData.OwnedObjects.Contains(FConcertObjectInStreamID{ SenderStreamId, ObjectReplicator->TestObject }));
+			}
+			else
+			{
+				AddError(TEXT("No activity was produced"));
+			}
+		}
+	END_DEFINE_SPEC(FLeaveReplicationActivitySpec);
+
+	/** This tests that when a client leaves replication, an activity containing their registered streams and authority is produced. */
+	void FLeaveReplicationActivitySpec::Define()
+	{
+		BeforeEach([this]
+		{
+			WorkspaceMock = MakeShared<FReplicationWorkspaceCallDetectorMock>();
+			ObjectReplicator = MakeUnique<FObjectTestReplicator>();
+			Server = MakeUnique<FReplicationServer>(*this, EConcertSyncSessionFlags::Default_MultiUserSession, WorkspaceMock.ToSharedRef());
+			Client = &Server->ConnectClient();
+			
+			Client->JoinReplication(ObjectReplicator->CreateSenderArgs(SenderStreamId));
+			Client->GetClientReplicationManager().TakeAuthorityOver({ ObjectReplicator->TestObject });
+		});
+		AfterEach([this]
+		{
+			Server.Reset();
+		});
+		
+		It("When leaving replication, an activity is produced", [this]
+		{
+			Client->LeaveReplication();
+			ValidateLeaveActivity();
+		});
+		It("When disconnecting from the session, an activity is produced", [this]
+		{
+			Client->GetClientSessionMock()->Disconnect();
+			ValidateLeaveActivity();
+		});
+	}
+	
+	BEGIN_DEFINE_SPEC(FNoReplicationActivitiesSpec, "Editor.Concert.Replication.Activity", EAutomationTestFlags::EditorContext | EAutomationTestFlags::EngineFilter)
+		/** Detects calls into the workspace */
+		TSharedPtr<FReplicationWorkspaceCallDetectorMock> WorkspaceMock;
+	
+		TUniquePtr<FReplicationServer> Server;
+		FReplicationClient* Client = nullptr;
+	END_DEFINE_SPEC(FNoReplicationActivitiesSpec);
+
+	/** This tests no activities are generated when the session EConcertSyncSessionFlags::ShouldEnableReplicationActivities is not set.  */
+	void FNoReplicationActivitiesSpec::Define()
+	{
+		BeforeEach([this]
+		{
+			WorkspaceMock = MakeShared<FReplicationWorkspaceCallDetectorMock>();
+			Server = MakeUnique<FReplicationServer>(*this, EConcertSyncSessionFlags::Default_MultiUserSession & ~EConcertSyncSessionFlags::ShouldEnableReplicationActivities, WorkspaceMock.ToSharedRef());
+			Client = &Server->ConnectClient();
+
+			Client->JoinReplicationAsListener({});
+		});
+		AfterEach([this]
+		{
+			Server.Reset();
+			WorkspaceMock.Reset();
+		});
+		
+		It("When leaving replication, no activity is produced", [this]
+		{
+			Client->LeaveReplication();
+			TestFalse(TEXT("Activity produced"), WorkspaceMock->LastCall_ProduceClientLeaveReplicationActivity.IsSet());
+		});
+	}
+}
diff --git a/Engine/Plugins/Developer/Concert/ConcertSync/ConcertSyncTest/Source/ConcertSyncTest/Private/Replication/ReplicationHandshakeTests.cpp b/Engine/Plugins/Developer/Concert/ConcertSync/ConcertSyncTest/Source/ConcertSyncTest/Private/Replication/ReplicationHandshakeTests.cpp
--- a/Engine/Plugins/Developer/Concert/ConcertSync/ConcertSyncTest/Source/ConcertSyncTest/Private/Replication/ReplicationHandshakeTests.cpp
+++ b/Engine/Plugins/Developer/Concert/ConcertSync/ConcertSyncTest/Source/ConcertSyncTest/Private/Replication/ReplicationHandshakeTests.cpp
@@ -14,6 +14,7 @@
 #include "Misc/AutomationTest.h"
 #include "Misc/Paths.h"
 #include "HAL/FileManager.h"
+#include "Util/Mocks/ReplicationWorkspaceEmptyMock.h"
 
 namespace UE::ConcertSyncTests::Replication::Handshake
 {
@@ -36,7 +37,8 @@ namespace UE::ConcertSyncTests::Replication::Handshake
 		// Server
 		InitServer();
 		const TSharedPtr<IConcertServerSession>& ServerSession = GetServerSessionMock();
-		const TSharedRef<IConcertServerReplicationManager> ServerReplicationManager = CreateServerReplicationManager(ServerSession.ToSharedRef());
+		FReplicationWorkspaceEmptyMock WorkspaceEmptyMock;
+		const TSharedRef<IConcertServerReplicationManager> ServerReplicationManager = CreateServerReplicationManager(ServerSession.ToSharedRef(), WorkspaceEmptyMock);
 		
 		// Client
 		FClientInfo& Client = ConnectClient();
diff --git a/Engine/Plugins/Developer/Concert/ConcertSync/ConcertSyncTest/Source/ConcertSyncTest/Private/Replication/ReplicationTestInterface.h b/Engine/Plugins/Developer/Concert/ConcertSync/ConcertSyncTest/Source/ConcertSyncTest/Private/Replication/ReplicationTestInterface.h
--- a/Engine/Plugins/Developer/Concert/ConcertSync/ConcertSyncTest/Source/ConcertSyncTest/Private/Replication/ReplicationTestInterface.h
+++ b/Engine/Plugins/Developer/Concert/ConcertSync/ConcertSyncTest/Source/ConcertSyncTest/Private/Replication/ReplicationTestInterface.h
@@ -13,6 +13,7 @@ class IConcertClientReplicationManager;
 namespace UE::ConcertSyncServer::Replication
 {
 	class IConcertServerReplicationManager;
+	class IReplicationWorkspace;
 }
 
 namespace UE::ConcertSyncClient::TestInterface
@@ -30,6 +31,7 @@ namespace UE::ConcertSyncServer::TestInterface
 {
 	extern CONCERTSYNCSERVER_API TSharedRef<Replication::IConcertServerReplicationManager> CreateServerReplicationManager(
 		TSharedRef<IConcertServerSession> InLiveSession,
+		Replication::IReplicationWorkspace& InWorkspace UE_LIFETIMEBOUND,
 		EConcertSyncSessionFlags InSessionFlags = EConcertSyncSessionFlags::Default_MultiUserSession
 		);
 }
\ No newline at end of file
diff --git a/Engine/Plugins/Developer/Concert/ConcertSync/ConcertSyncTest/Source/ConcertSyncTest/Private/Replication/Util/Mocks/ReplicationWorkspaceCallDetectorMock.h b/Engine/Plugins/Developer/Concert/ConcertSync/ConcertSyncTest/Source/ConcertSyncTest/Private/Replication/Util/Mocks/ReplicationWorkspaceCallDetectorMock.h
new file mode 100644
--- /dev/null
+++ b/Engine/Plugins/Developer/Concert/ConcertSync/ConcertSyncTest/Source/ConcertSyncTest/Private/Replication/Util/Mocks/ReplicationWorkspaceCallDetectorMock.h
@@ -0,0 +1,28 @@
+﻿// Copyright Epic Games, Inc. All Rights Reserved.
+
+#pragma once
+
+#include "Replication/IReplicationWorkspace.h"
+#include "Replication/Messages/ReplicationActivity.h"
+
+#include "Misc/Optional.h"
+#include "Misc/Guid.h"
+#include "Templates/Tuple.h"
+
+namespace UE::ConcertSyncTests::Replication
+{
+	class FReplicationWorkspaceCallDetectorMock : public ConcertSyncServer::Replication::IReplicationWorkspace
+	{
+	public:
+
+		/** Arguments of last ProduceClientLeaveReplicationActivity call. */
+		TOptional<TTuple<FGuid, FConcertSyncReplicationPayload_LeaveReplication>> LastCall_ProduceClientLeaveReplicationActivity;
+		
+		//~ Begin IReplicationWorkspace Interface
+		virtual void ProduceClientLeaveReplicationActivity(const FGuid& EndpointId, const FConcertSyncReplicationPayload_LeaveReplication& EventData) override
+		{
+			LastCall_ProduceClientLeaveReplicationActivity = MakeTuple(EndpointId, EventData);
+		}
+		//~ End IReplicationWorkspace Interface
+	};
+}
diff --git a/Engine/Plugins/Developer/Concert/ConcertSync/ConcertSyncTest/Source/ConcertSyncTest/Private/Replication/Util/Mocks/ReplicationWorkspaceEmptyMock.h b/Engine/Plugins/Developer/Concert/ConcertSync/ConcertSyncTest/Source/ConcertSyncTest/Private/Replication/Util/Mocks/ReplicationWorkspaceEmptyMock.h
new file mode 100644
--- /dev/null
+++ b/Engine/Plugins/Developer/Concert/ConcertSync/ConcertSyncTest/Source/ConcertSyncTest/Private/Replication/Util/Mocks/ReplicationWorkspaceEmptyMock.h
@@ -0,0 +1,16 @@
+﻿// Copyright Epic Games, Inc. All Rights Reserved.
+
+#pragma once
+#include "Replication/IReplicationWorkspace.h"
+
+namespace UE::ConcertSyncTests::Replication
+{
+	class FReplicationWorkspaceEmptyMock : public ConcertSyncServer::Replication::IReplicationWorkspace
+	{
+	public:
+		
+		//~ Begin IReplicationWorkspace Interface
+		virtual void ProduceClientLeaveReplicationActivity(const FGuid&, const FConcertSyncReplicationPayload_LeaveReplication&) override {}
+		//~ End IReplicationWorkspace Interface
+	};
+}
diff --git a/Engine/Plugins/Developer/Concert/ConcertSync/ConcertSyncTest/Source/ConcertSyncTest/Private/Replication/Util/SendReceiveTestBase.cpp b/Engine/Plugins/Developer/Concert/ConcertSync/ConcertSyncTest/Source/ConcertSyncTest/Private/Replication/Util/SendReceiveTestBase.cpp
--- a/Engine/Plugins/Developer/Concert/ConcertSync/ConcertSyncTest/Source/ConcertSyncTest/Private/Replication/Util/SendReceiveTestBase.cpp
+++ b/Engine/Plugins/Developer/Concert/ConcertSync/ConcertSyncTest/Source/ConcertSyncTest/Private/Replication/Util/SendReceiveTestBase.cpp
@@ -51,7 +51,7 @@ namespace UE::ConcertSyncTests::Replication
 		// Server
 		InitServer();
 		ServerSession = GetServerSessionMock();
-		ServerReplicationManager = ConcertSyncServer::TestInterface::CreateServerReplicationManager(ServerSession.ToSharedRef());
+		ServerReplicationManager = ConcertSyncServer::TestInterface::CreateServerReplicationManager(ServerSession.ToSharedRef(), WorkspaceMock);
 		// Client Receiver
 		Client_Receiver = &ConnectClient();
 		if (EnumHasAnyFlags(Flags, ESendReceiveTestFlags::UseRealReplicationBridge))
diff --git a/Engine/Plugins/Developer/Concert/ConcertSync/ConcertSyncTest/Source/ConcertSyncTest/Private/Replication/Util/SendReceiveTestBase.h b/Engine/Plugins/Developer/Concert/ConcertSync/ConcertSyncTest/Source/ConcertSyncTest/Private/Replication/Util/SendReceiveTestBase.h
--- a/Engine/Plugins/Developer/Concert/ConcertSync/ConcertSyncTest/Source/ConcertSyncTest/Private/Replication/Util/SendReceiveTestBase.h
+++ b/Engine/Plugins/Developer/Concert/ConcertSync/ConcertSyncTest/Source/ConcertSyncTest/Private/Replication/Util/SendReceiveTestBase.h
@@ -2,6 +2,7 @@
 
 #pragma once
 
+#include "Mocks/ReplicationWorkspaceEmptyMock.h"
 #include "Replication/IConcertClientReplicationManager.h"
 #include "Replication/Messages/ObjectReplication.h"
 #include "Util/ClientServerCommunicationTest.h"
@@ -60,6 +61,7 @@ namespace UE::ConcertSyncTests::Replication
 		float FakeDeltaTime = 1.f / 60.f;
 		
 		TSharedPtr<FConcertServerSessionMock> ServerSession;
+		FReplicationWorkspaceEmptyMock WorkspaceMock;
 		TSharedPtr<ConcertSyncServer::Replication::IConcertServerReplicationManager> ServerReplicationManager;
 		
 		FClientInfo* Client_Receiver = nullptr;
diff --git a/Engine/Plugins/Developer/Concert/ConcertSync/ConcertSyncTest/Source/ConcertSyncTest/Private/Replication/Util/Spec/ReplicationClient.cpp b/Engine/Plugins/Developer/Concert/ConcertSync/ConcertSyncTest/Source/ConcertSyncTest/Private/Replication/Util/Spec/ReplicationClient.cpp
--- a/Engine/Plugins/Developer/Concert/ConcertSync/ConcertSyncTest/Source/ConcertSyncTest/Private/Replication/Util/Spec/ReplicationClient.cpp
+++ b/Engine/Plugins/Developer/Concert/ConcertSync/ConcertSyncTest/Source/ConcertSyncTest/Private/Replication/Util/Spec/ReplicationClient.cpp
@@ -67,5 +67,10 @@ namespace UE::ConcertSyncTests::Replication
 				return Result;
 			});
 	}
+
+	void FReplicationClient::LeaveReplication() const
+	{
+		ClientReplicationManager->LeaveReplicationSession();
+	}
 }
 
diff --git a/Engine/Plugins/Developer/Concert/ConcertSync/ConcertSyncTest/Source/ConcertSyncTest/Private/Replication/Util/Spec/ReplicationClient.h b/Engine/Plugins/Developer/Concert/ConcertSync/ConcertSyncTest/Source/ConcertSyncTest/Private/Replication/Util/Spec/ReplicationClient.h
--- a/Engine/Plugins/Developer/Concert/ConcertSync/ConcertSyncTest/Source/ConcertSyncTest/Private/Replication/Util/Spec/ReplicationClient.h
+++ b/Engine/Plugins/Developer/Concert/ConcertSync/ConcertSyncTest/Source/ConcertSyncTest/Private/Replication/Util/Spec/ReplicationClient.h
@@ -35,7 +35,8 @@ namespace UE::ConcertSyncTests::Replication
 	public:
 		
 		FReplicationClient(const FGuid& ClientEndPointId, EConcertSyncSessionFlags SessionFlags, FConcertServerSessionMock& Server, FAutomationTestBase& TestContext)
-			: SessionFlags(SessionFlags)
+			: ClientEndPointId(ClientEndPointId)
+			, SessionFlags(SessionFlags)
 			, TestContext(TestContext)
 			, ClientSessionMock(MakeShared<FConcertClientSessionMock>(ClientEndPointId, Server))
 		{}
@@ -53,13 +54,18 @@ namespace UE::ConcertSyncTests::Replication
 			TArray<UObject*> ObjectsToReceive
 			);
 
-		TSharedRef<FConcertClientSessionBaseMock> GetClientSessionMock() const { return ClientSessionMock; }
+		/** Leaves the session's replication. */
+		void LeaveReplication() const;
 
+		const FGuid& GetEndpointId() const { return ClientEndPointId; }
+		TSharedRef<FConcertClientSessionBaseMock> GetClientSessionMock() const { return ClientSessionMock; }
 		FConcertClientReplicationBridgeMock& GetBridgeMock() const { checkf(BridgeMock, TEXT("You forgot to call JoinReplication")); return *BridgeMock; }
 		IConcertClientReplicationManager& GetClientReplicationManager() const { checkf(ClientReplicationManager, TEXT("You forgot to call JoinReplication")); return *ClientReplicationManager; }
 
 	private:
 
+		/** The client's endpoint. */
+		const FGuid ClientEndPointId;
 		/** Relevant for certain requests. Passed to ClientReplicationManager upon creation. */
 		const EConcertSyncSessionFlags SessionFlags;
 
diff --git a/Engine/Plugins/Developer/Concert/ConcertSync/ConcertSyncTest/Source/ConcertSyncTest/Private/Replication/Util/Spec/ReplicationServer.cpp b/Engine/Plugins/Developer/Concert/ConcertSync/ConcertSyncTest/Source/ConcertSyncTest/Private/Replication/Util/Spec/ReplicationServer.cpp
--- a/Engine/Plugins/Developer/Concert/ConcertSync/ConcertSyncTest/Source/ConcertSyncTest/Private/Replication/Util/Spec/ReplicationServer.cpp
+++ b/Engine/Plugins/Developer/Concert/ConcertSync/ConcertSyncTest/Source/ConcertSyncTest/Private/Replication/Util/Spec/ReplicationServer.cpp
@@ -7,11 +7,16 @@
 
 namespace UE::ConcertSyncTests::Replication
 {
-	FReplicationServer::FReplicationServer(FAutomationTestBase& TestContext, EConcertSyncSessionFlags InSessionFlags)
+	FReplicationServer::FReplicationServer(
+		FAutomationTestBase& TestContext,
+		EConcertSyncSessionFlags InSessionFlags,
+		TSharedRef<ConcertSyncServer::Replication::IReplicationWorkspace> InWorkspace
+		)
 		: SessionFlags(InSessionFlags)
 		, TestContext(TestContext)
 		, ServerSessionMock(MakeShared<FConcertServerSessionMock>())
-		, ServerReplicationManager(ConcertSyncServer::TestInterface::CreateServerReplicationManager(ServerSessionMock, InSessionFlags))
+		, ReplicationWorkspace(MoveTemp(InWorkspace))
+		, ServerReplicationManager(ConcertSyncServer::TestInterface::CreateServerReplicationManager(ServerSessionMock, *ReplicationWorkspace, InSessionFlags))
 	{}
 
 	FReplicationClient& FReplicationServer::ConnectClient()
diff --git a/Engine/Plugins/Developer/Concert/ConcertSync/ConcertSyncTest/Source/ConcertSyncTest/Private/Replication/Util/Spec/ReplicationServer.h b/Engine/Plugins/Developer/Concert/ConcertSync/ConcertSyncTest/Source/ConcertSyncTest/Private/Replication/Util/Spec/ReplicationServer.h
--- a/Engine/Plugins/Developer/Concert/ConcertSync/ConcertSyncTest/Source/ConcertSyncTest/Private/Replication/Util/Spec/ReplicationServer.h
+++ b/Engine/Plugins/Developer/Concert/ConcertSync/ConcertSyncTest/Source/ConcertSyncTest/Private/Replication/Util/Spec/ReplicationServer.h
@@ -4,10 +4,14 @@
 
 #include "ConcertSyncSessionFlags.h"
 #include "Replication/IConcertServerReplicationManager.h"
+#include "Replication/Util/Mocks/ReplicationWorkspaceEmptyMock.h"
 #include "Util/ClientServerCommunicationTest.h"
 
 #include "Templates/UnrealTemplate.h"
 
+namespace UE::ConcertSyncServer::Replication { class IReplicationWorkspace;
+}
+
 class FAutomationTestBase;
 
 namespace UE::ConcertSyncTests::Replication
@@ -19,7 +23,11 @@ namespace UE::ConcertSyncTests::Replication
 	{
 	public:
 
-		FReplicationServer(FAutomationTestBase& TestContext, EConcertSyncSessionFlags InSessionFlags = EConcertSyncSessionFlags::Default_MultiUserSession);
+		FReplicationServer(
+			FAutomationTestBase& TestContext,
+			EConcertSyncSessionFlags InSessionFlags = EConcertSyncSessionFlags::Default_MultiUserSession,
+			TSharedRef<ConcertSyncServer::Replication::IReplicationWorkspace> InWorkspace = MakeShared<FReplicationWorkspaceEmptyMock>()
+			);
 
 		/** Connects a client to the server. */
 		FReplicationClient& ConnectClient();
@@ -39,6 +47,8 @@ namespace UE::ConcertSyncTests::Replication
 
 		/** The underlying server session */
 		TSharedRef<FConcertServerSessionMock> ServerSessionMock;
+		/** Mock for the server workspace the replication system interacts with. */
+		TSharedRef<ConcertSyncServer::Replication::IReplicationWorkspace> ReplicationWorkspace;
 		/** Manages replication server side */
 		TSharedRef<ConcertSyncServer::Replication::IConcertServerReplicationManager> ServerReplicationManager;
 
