Added FConcertReplication_RestoreContent_Request. Clients can send this request to the server to revert to a state they had when they left the session.
A later CL will use this to auto-restore session data upon joining replication.

diff --git a/Engine/Plugins/Developer/Concert/ConcertSync/ConcertSyncClient/Source/ConcertSyncClient/Private/Replication/Manager/ReplicationManager.cpp b/Engine/Plugins/Developer/Concert/ConcertSync/ConcertSyncClient/Source/ConcertSyncClient/Private/Replication/Manager/ReplicationManager.cpp
--- a/Engine/Plugins/Developer/Concert/ConcertSync/ConcertSyncClient/Source/ConcertSyncClient/Private/Replication/Manager/ReplicationManager.cpp
+++ b/Engine/Plugins/Developer/Concert/ConcertSync/ConcertSyncClient/Source/ConcertSyncClient/Private/Replication/Manager/ReplicationManager.cpp
@@ -140,6 +140,13 @@ namespace UE::ConcertSyncClient::Replication
 			: MakeFulfilledPromise<FConcertReplication_QueryMuteState_Response>().GetFuture(); 
 	}
 
+	TFuture<FConcertReplication_RestoreContent_Response> FReplicationManager::RestoreContent(FConcertReplication_RestoreContent_Request Request)
+	{
+		return ensureMsgf(CurrentState, TEXT("StartAcceptingJoinRequests should have been called at this point."))
+			? CurrentState->RestoreContent(MoveTemp(Request))
+			: MakeFulfilledPromise<FConcertReplication_RestoreContent_Response>().GetFuture(); 
+	}
+
 	IConcertClientReplicationManager::FOnPreStreamsChanged& FReplicationManager::OnPreStreamsChanged()
 	{
 		// Check() to avoid returning some dummy static variable
diff --git a/Engine/Plugins/Developer/Concert/ConcertSync/ConcertSyncClient/Source/ConcertSyncClient/Private/Replication/Manager/ReplicationManager.h b/Engine/Plugins/Developer/Concert/ConcertSync/ConcertSyncClient/Source/ConcertSyncClient/Private/Replication/Manager/ReplicationManager.h
--- a/Engine/Plugins/Developer/Concert/ConcertSync/ConcertSyncClient/Source/ConcertSyncClient/Private/Replication/Manager/ReplicationManager.h
+++ b/Engine/Plugins/Developer/Concert/ConcertSync/ConcertSyncClient/Source/ConcertSyncClient/Private/Replication/Manager/ReplicationManager.h
@@ -45,6 +45,7 @@ namespace UE::ConcertSyncClient::Replication
 		virtual bool HasSyncControl(const FConcertObjectInStreamID& Object) const override;
 		virtual TFuture<FConcertReplication_ChangeMuteState_Response> ChangeMuteState(FConcertReplication_ChangeMuteState_Request Request) override;
 		virtual TFuture<FConcertReplication_QueryMuteState_Response> QueryMuteState(FConcertReplication_QueryMuteState_Request Request) override;
+		virtual TFuture<FConcertReplication_RestoreContent_Response> RestoreContent(FConcertReplication_RestoreContent_Request Request) override;
 		virtual FOnPreStreamsChanged& OnPreStreamsChanged() override;
 		virtual FOnPostStreamsChanged& OnPostStreamsChanged() override;
 		virtual FOnPreAuthorityChanged& OnPreAuthorityChanged() override;
diff --git a/Engine/Plugins/Developer/Concert/ConcertSync/ConcertSyncClient/Source/ConcertSyncClient/Private/Replication/Manager/ReplicationManagerState.h b/Engine/Plugins/Developer/Concert/ConcertSync/ConcertSyncClient/Source/ConcertSyncClient/Private/Replication/Manager/ReplicationManagerState.h
--- a/Engine/Plugins/Developer/Concert/ConcertSync/ConcertSyncClient/Source/ConcertSyncClient/Private/Replication/Manager/ReplicationManagerState.h
+++ b/Engine/Plugins/Developer/Concert/ConcertSync/ConcertSyncClient/Source/ConcertSyncClient/Private/Replication/Manager/ReplicationManagerState.h
@@ -34,6 +34,7 @@ namespace UE::ConcertSyncClient::Replication
 		virtual bool HasSyncControl(const FConcertObjectInStreamID& Object) const override { return false; }
 		virtual TFuture<FConcertReplication_ChangeMuteState_Response> ChangeMuteState(FConcertReplication_ChangeMuteState_Request) override { return MakeFulfilledPromise<FConcertReplication_ChangeMuteState_Response>(FConcertReplication_ChangeMuteState_Response{ EConcertReplicationMuteErrorCode::Rejected }).GetFuture(); };
 		virtual TFuture<FConcertReplication_QueryMuteState_Response> QueryMuteState(FConcertReplication_QueryMuteState_Request Request) override { return MakeFulfilledPromise<FConcertReplication_QueryMuteState_Response>().GetFuture(); }
+		virtual TFuture<FConcertReplication_RestoreContent_Response> RestoreContent(FConcertReplication_RestoreContent_Request Request) override { return MakeFulfilledPromise<FConcertReplication_RestoreContent_Response>().GetFuture(); }
 		virtual FOnPreStreamsChanged& OnPreStreamsChanged() override { return OnPreStreamsChangedDelegate; } 
 		virtual FOnPostStreamsChanged& OnPostStreamsChanged() override { return OnPostStreamsChangedDelegate; }
 		virtual FOnPreAuthorityChanged& OnPreAuthorityChanged() override { return OnPreAuthorityChangedDelegate; }
diff --git a/Engine/Plugins/Developer/Concert/ConcertSync/ConcertSyncClient/Source/ConcertSyncClient/Private/Replication/Manager/ReplicationManagerState_Connected.cpp b/Engine/Plugins/Developer/Concert/ConcertSync/ConcertSyncClient/Source/ConcertSyncClient/Private/Replication/Manager/ReplicationManagerState_Connected.cpp
--- a/Engine/Plugins/Developer/Concert/ConcertSync/ConcertSyncClient/Source/ConcertSyncClient/Private/Replication/Manager/ReplicationManagerState_Connected.cpp
+++ b/Engine/Plugins/Developer/Concert/ConcertSync/ConcertSyncClient/Source/ConcertSyncClient/Private/Replication/Manager/ReplicationManagerState_Connected.cpp
@@ -63,6 +63,11 @@ namespace UE::ConcertSyncClient::Replication
 		false,
 		TEXT("Whether to log changes to the mute state.")
 		);
+	TAutoConsoleVariable<bool> CVarLogRestoreContentRequestsAndResponsesOnClient(
+		TEXT("Concert.Replication.LogRestoreContentRequestsAndResponsesOnClient"),
+		false,
+		TEXT("Whether to log restore content requests and responses.")
+		);
 
 	namespace Private
 	{
@@ -232,7 +237,7 @@ namespace UE::ConcertSyncClient::Replication
 				if (ThisPin && Response.IsSuccess())
 				{
 					ThisPin->SyncControl.ProcessStreamChange(Args);
-					ThisPin->UpdateReplicatedObjectsAfterStreamChange(Args, Response);
+					ThisPin->UpdateReplicatedObjectsAfterStreamChange(Args);
 				}
 				else if (ThisPin && Response.ErrorCode == EReplicationResponseErrorCode::Timeout)
 				{
@@ -292,12 +297,10 @@ namespace UE::ConcertSyncClient::Replication
 			.Next([WeakThis = TWeakPtr<FReplicationManagerState_Connected>(SharedThis(this)), PredictedChanges = MoveTemp(PredictedChanges)](FConcertReplication_ChangeMuteState_Response&& Response)
 			{
 				Private::LogNetworkMessage(CVarLogMuteRequestsAndResponsesOnClient, Response);
-				
 				if (const TSharedPtr<FReplicationManagerState_Connected> This = WeakThis.Pin())
 				{
 					This->SyncControl.ApplyOrRevertMuteResponse(PredictedChanges, Response);
 				}
-				
 				return Response;
 			});
 	}
@@ -315,6 +318,51 @@ namespace UE::ConcertSyncClient::Replication
 			);
 	}
 
+	TFuture<FConcertReplication_RestoreContent_Response> FReplicationManagerState_Connected::RestoreContent(FConcertReplication_RestoreContent_Request Request)
+	{
+		if (!EnumHasAnyFlags(SessionFlags, EConcertSyncSessionFlags::ShouldEnableReplicationActivities))
+		{
+			return MakeFulfilledPromise<FConcertReplication_RestoreContent_Response>(FConcertReplication_RestoreContent_Response{ EConcertReplicationRestoreErrorCode::NotSupported }).GetFuture();
+		}
+		
+		FLocalSyncControl::FPredictedObjectRemoval PredictedChanges = SyncControl.PredictAndApplyRestoreContentChanges(Request);
+
+		// We want the response to contain ClientInfo to update our internal state - so set the SendNewState flag
+		const bool bWantedNewState = EnumHasAnyFlags(Request.Flags, EConcertReplicationRestoreContentFlags::SendNewState);
+		Request.Flags |= EConcertReplicationRestoreContentFlags::SendNewState;
+		Private::LogNetworkMessage(CVarLogRestoreContentRequestsAndResponsesOnClient, Request);
+		
+		return LiveSession->SendCustomRequest<FConcertReplication_RestoreContent_Request, FConcertReplication_RestoreContent_Response>(Request, LiveSession->GetSessionServerEndpointId())
+			.Next([WeakThis = TWeakPtr<FReplicationManagerState_Connected>(SharedThis(this)), PredictedChanges = MoveTemp(PredictedChanges)](FConcertReplication_RestoreContent_Response&& Response)
+			{
+				Private::LogNetworkMessage(CVarLogRestoreContentRequestsAndResponsesOnClient, Response);
+				
+				const TSharedPtr<FReplicationManagerState_Connected> This = WeakThis.Pin();
+				if (!This)
+				{
+					return Response;
+				}
+
+				This->SyncControl.ApplyOrRevertRestoreContentResponse(PredictedChanges, Response);
+				if (Response.IsSuccess())
+				{
+					// Update the list of objects we'll be replicating. This is why we added the SendNewState flag above.
+					This->UpdateReplicatedObjectAfterServerSideChange(Response.ClientInfo);
+				}
+				
+				return Response;
+			})
+			.Next([bWantedNewState](FConcertReplication_RestoreContent_Response&& Response)
+			{
+				// Since we added EConcertReplicationRestoreContentFlags::SendNewState above, remove the returned data if the flag was not originally in the request.
+				if (!bWantedNewState)
+				{
+					Response.ClientInfo = {};
+				}
+				return Response;
+			});
+	}
+
 	void FReplicationManagerState_Connected::OnEnterState()
 	{
 		LiveSession->OnTick().AddSP(this, &FReplicationManagerState_Connected::Tick);
@@ -328,9 +376,9 @@ namespace UE::ConcertSyncClient::Replication
 		ReplicationApplier.ProcessObjects(Params);
 	}
 
-	void FReplicationManagerState_Connected::UpdateReplicatedObjectsAfterStreamChange(const FConcertReplication_ChangeStream_Request& Request, const FConcertReplication_ChangeStream_Response& Response)
+	void FReplicationManagerState_Connected::UpdateReplicatedObjectsAfterStreamChange(const FConcertReplication_ChangeStream_Request& Request)
 	{
-		OnPreStreamsChangedDelegate.Broadcast(Request, { Response });
+		OnPreStreamsChangedDelegate.Broadcast();
 		ON_SCOPE_EXIT{ OnPostStreamsChangedDelegate.Broadcast(); };
 		
 		// Build RegisteredStreams while RegisteredStreams has the old, unupdated state
@@ -397,7 +445,7 @@ namespace UE::ConcertSyncClient::Replication
 
 	void FReplicationManagerState_Connected::UpdateReplicatedObjectsAfterAuthorityChange(FConcertReplication_ChangeAuthority_Request&& Request, const FConcertReplication_ChangeAuthority_Response& Response)
 	{
-		OnPreAuthorityChangedDelegate.Broadcast(Request, { Response });
+		OnPreAuthorityChangedDelegate.Broadcast();
 		ON_SCOPE_EXIT{ OnPostAuthorityChangedDelegate.Broadcast(); };
 		
 		for (TPair<FSoftObjectPath, FConcertStreamArray>& TakeAuthority : Request.TakeAuthority)
@@ -441,6 +489,34 @@ namespace UE::ConcertSyncClient::Replication
 		}
 	}
 
+	void FReplicationManagerState_Connected::UpdateReplicatedObjectAfterServerSideChange(const FConcertQueriedClientInfo& NewState)
+	{
+		OnPreStreamsChangedDelegate.Broadcast();
+		OnPreAuthorityChangedDelegate.Broadcast();
+		ON_SCOPE_EXIT
+		{
+			OnPostStreamsChangedDelegate.Broadcast();
+			OnPostAuthorityChangedDelegate.Broadcast();
+		};
+		
+		// The list of registered streams has changed
+		RegisteredStreams.Empty(NewState.Streams.Num());
+		Algo::Transform(NewState.Streams, RegisteredStreams, [](const FConcertBaseStreamInfo& Info)
+		{
+			return FConcertReplicationStream{ Info };
+		});
+
+		// This affects the objects we'll be replicating
+		ReplicationDataSource.ClearReplicatedObjects();
+		for (const FConcertAuthorityClientInfo& AuthorityState : NewState.Authority)
+		{
+			for (const FSoftObjectPath& ObjectPath : AuthorityState.AuthoredObjects)
+			{
+				ReplicationDataSource.AddReplicatedObjectStreams(ObjectPath, { AuthorityState.StreamId });
+			}
+		}
+	}
+
 	FConcertObjectReplicationSettings FReplicationManagerState_Connected::GetObjectFrequencySettings(const FConcertReplicatedObjectId& Object) const
 	{
 		const FConcertReplicationStream* Stream = RegisteredStreams.FindByPredicate([&Object](const FConcertReplicationStream& Description)
diff --git a/Engine/Plugins/Developer/Concert/ConcertSync/ConcertSyncClient/Source/ConcertSyncClient/Private/Replication/Manager/ReplicationManagerState_Connected.h b/Engine/Plugins/Developer/Concert/ConcertSync/ConcertSyncClient/Source/ConcertSyncClient/Private/Replication/Manager/ReplicationManagerState_Connected.h
--- a/Engine/Plugins/Developer/Concert/ConcertSync/ConcertSyncClient/Source/ConcertSyncClient/Private/Replication/Manager/ReplicationManagerState_Connected.h
+++ b/Engine/Plugins/Developer/Concert/ConcertSync/ConcertSyncClient/Source/ConcertSyncClient/Private/Replication/Manager/ReplicationManagerState_Connected.h
@@ -56,6 +56,7 @@ namespace UE::ConcertSyncClient::Replication
 		virtual bool HasSyncControl(const FConcertObjectInStreamID& Object) const override { return SyncControl.IsObjectAllowed(Object); }
 		virtual TFuture<FConcertReplication_ChangeMuteState_Response> ChangeMuteState(FConcertReplication_ChangeMuteState_Request) override;
 		virtual TFuture<FConcertReplication_QueryMuteState_Response> QueryMuteState(FConcertReplication_QueryMuteState_Request Request) override;
+		virtual TFuture<FConcertReplication_RestoreContent_Response> RestoreContent(FConcertReplication_RestoreContent_Request Request) override;
 		//~ End IConcertClientReplicationManager Interface
 
 	private:
@@ -110,7 +111,7 @@ namespace UE::ConcertSyncClient::Replication
 		void Tick(IConcertClientSession& Session, float DeltaTime);
 		
 		/** Updates replicated objects affected by the change request. */
-		void UpdateReplicatedObjectsAfterStreamChange(const FConcertReplication_ChangeStream_Request& Request, const FConcertReplication_ChangeStream_Response& Response);
+		void UpdateReplicatedObjectsAfterStreamChange(const FConcertReplication_ChangeStream_Request& Request);
 		void HandleRemovingReplicatedObjects(const FConcertReplication_ChangeStream_Request& Request);
 		void RevertRemovingReplicatedObjects(const FConcertReplication_ChangeStream_Request& Request);
 
@@ -118,11 +119,14 @@ namespace UE::ConcertSyncClient::Replication
 		 * Updates the objects which should be replicated after changing authority.
 		 * 
 		 * @note Request is accepted as && because this function rewrites its memory when looking at rejections.
-		 * Since the request was already sent to the server it is assumed the request can just contain trash after.
+		 * Since the request was already sent to the server it can just contain trash after.
 		 */
 		void UpdateReplicatedObjectsAfterAuthorityChange(FConcertReplication_ChangeAuthority_Request&& Request, const FConcertReplication_ChangeAuthority_Response& Response);
 		void HandleReleasingReplicatedObjects(const FConcertReplication_ChangeAuthority_Request& Request);
 		void RevertReleasingReplicatedObjects(const FConcertReplication_ChangeAuthority_Request& Request);
+
+		/** Updates the objects which should be replicated after they have been reset to a completely new state (e.g. when restoring session content manually). */
+		void UpdateReplicatedObjectAfterServerSideChange(const FConcertQueriedClientInfo& NewState);
 		
 		/** Callback to Sender for obtaining an object's frequency settings. */
 		FConcertObjectReplicationSettings GetObjectFrequencySettings(const FConcertReplicatedObjectId& Object) const;
diff --git a/Engine/Plugins/Developer/Concert/ConcertSync/ConcertSyncClient/Source/ConcertSyncClient/Private/Replication/Misc/LocalSyncControl.h b/Engine/Plugins/Developer/Concert/ConcertSync/ConcertSyncClient/Source/ConcertSyncClient/Private/Replication/Misc/LocalSyncControl.h
--- a/Engine/Plugins/Developer/Concert/ConcertSync/ConcertSyncClient/Source/ConcertSyncClient/Private/Replication/Misc/LocalSyncControl.h
+++ b/Engine/Plugins/Developer/Concert/ConcertSync/ConcertSyncClient/Source/ConcertSyncClient/Private/Replication/Misc/LocalSyncControl.h
@@ -21,6 +21,7 @@ namespace UE::ConcertSyncClient::Replication
 		: public FNoncopyable
 		, ConcertSyncCore::Replication::FSyncControlState
 	{
+		using Detail = FSyncControlState;
 	public:
 
 		explicit FLocalSyncControl(IConcertSession& InSession UE_LIFETIMEBOUND)
@@ -53,7 +54,7 @@ namespace UE::ConcertSyncClient::Replication
 		void ProcessAuthorityChange(const FConcertReplication_ChangeAuthority_Request& Request, const FConcertReplication_ChangeAuthority_Response& Response)
 		{
 			OnPreSyncControlChangedDelegate.Broadcast();
-			AppendChanges(Request, Response);
+			AppendAuthorityChange(Request, Response);
 			OnPostSyncControlChangedDelegate.Broadcast();
 		}
 
@@ -64,7 +65,7 @@ namespace UE::ConcertSyncClient::Replication
 		void ProcessStreamChange(const FConcertReplication_ChangeStream_Request& Request)
 		{
 			OnPreSyncControlChangedDelegate.Broadcast();
-			AppendChanges(Request);
+			AppendStreamChange(Request);
 			OnPostSyncControlChangedDelegate.Broadcast();
 		}
 
@@ -75,7 +76,7 @@ namespace UE::ConcertSyncClient::Replication
 		FPredictedObjectRemoval PredictAndApplyMuteChanges(const FConcertReplication_ChangeMuteState_Request& Request)
 		{
 			bool bMadeChange = false;
-			FPredictedObjectRemoval Predicition = AppendChanges(Request, [this, &bMadeChange](const FConcertObjectInStreamID&)
+			FPredictedObjectRemoval Predicition = Detail::PredictAndApplyMuteChanges(Request, [this, &bMadeChange](const FConcertObjectInStreamID&)
 			{
 				if (!bMadeChange)
 				{
@@ -95,7 +96,7 @@ namespace UE::ConcertSyncClient::Replication
 		void ApplyOrRevertMuteResponse(const FPredictedObjectRemoval& RemovedByRequest, const FConcertReplication_ChangeMuteState_Response& Response)
 		{
 			bool bMadeChange = false;
-			AppendChanges(RemovedByRequest, Response, [this, &bMadeChange](const FConcertObjectInStreamID&)
+			Detail::ApplyOrRevertMuteResponse(RemovedByRequest, Response, [this, &bMadeChange](const FConcertObjectInStreamID&)
 			{
 				if (!bMadeChange)
 				{
@@ -108,6 +109,51 @@ namespace UE::ConcertSyncClient::Replication
 				OnPostSyncControlChangedDelegate.Broadcast();
 			}
 		}
+
+		
+		/**
+		 * Applies the implicit changes made by the request assuming the request will be accepted.
+		 * @return The removed objects to be passed to ApplyOrRevertRestoreContentResponse.
+		 */
+		FPredictedObjectRemoval PredictAndApplyRestoreContentChanges(const FConcertReplication_RestoreContent_Request& Request)
+		{
+			bool bMadeChange = false;
+			FPredictedObjectRemoval Predicition = Detail::PredictAndApplyRestoreContentChanges(Request, [this, &bMadeChange](const FConcertObjectInStreamID&)
+			{
+				if (!bMadeChange)
+				{
+					OnPreSyncControlChangedDelegate.Broadcast();
+				}
+				bMadeChange = true;
+			});
+			if (bMadeChange)
+			{
+				OnPostSyncControlChangedDelegate.Broadcast();
+			}
+
+			return Predicition;
+		}
+
+		/** Either reverts previous changes made if the request was rejected, or applies the sync control returned by the server otherwise. */
+		void ApplyOrRevertRestoreContentResponse(const FPredictedObjectRemoval& RemovedByRequest, const FConcertReplication_RestoreContent_Response& Response)
+		{
+			bool bMadeChange = false;
+			const auto OnChange = [this, &bMadeChange](const FConcertObjectInStreamID&)
+			{
+				if (!bMadeChange)
+				{
+					OnPreSyncControlChangedDelegate.Broadcast();
+				}
+				bMadeChange = true;
+			};
+			Detail::ApplyOrRevertRestoreContentResponse(RemovedByRequest, Response, OnChange, OnChange);
+			
+			if (bMadeChange)
+			{
+				OnPostSyncControlChangedDelegate.Broadcast();
+			}
+		}
+		
 		
 		DECLARE_MULTICAST_DELEGATE(FSyncControlChanged);
 		FSyncControlChanged& OnPreSyncControlChanged() { return OnPreSyncControlChangedDelegate; }
diff --git a/Engine/Plugins/Developer/Concert/ConcertSync/ConcertSyncClient/Source/ConcertSyncClient/Private/Replication/Processing/ClientReplicationDataCollector.cpp b/Engine/Plugins/Developer/Concert/ConcertSync/ConcertSyncClient/Source/ConcertSyncClient/Private/Replication/Processing/ClientReplicationDataCollector.cpp
--- a/Engine/Plugins/Developer/Concert/ConcertSync/ConcertSyncClient/Source/ConcertSyncClient/Private/Replication/Processing/ClientReplicationDataCollector.cpp
+++ b/Engine/Plugins/Developer/Concert/ConcertSync/ConcertSyncClient/Source/ConcertSyncClient/Private/Replication/Processing/ClientReplicationDataCollector.cpp
@@ -131,6 +131,15 @@ namespace UE::ConcertSyncClient::Replication
 		}
 	}
 
+	void FClientReplicationDataCollector::ClearReplicatedObjects()
+	{
+		for (const TPair<FSoftObjectPath, TArray<FObjectInfo>>& Pair : ObjectsToReplicate)
+		{
+			Bridge.PopTrackedObjects({ Pair.Key });
+		}
+		ObjectsToReplicate.Reset();
+	}
+
 	void FClientReplicationDataCollector::ForEachOwnedObject(
 		TFunctionRef<EBreakBehavior(const FSoftObjectPath&)> Callback
 		) const
diff --git a/Engine/Plugins/Developer/Concert/ConcertSync/ConcertSyncClient/Source/ConcertSyncClient/Private/Replication/Processing/ClientReplicationDataCollector.h b/Engine/Plugins/Developer/Concert/ConcertSync/ConcertSyncClient/Source/ConcertSyncClient/Private/Replication/Processing/ClientReplicationDataCollector.h
--- a/Engine/Plugins/Developer/Concert/ConcertSync/ConcertSyncClient/Source/ConcertSyncClient/Private/Replication/Processing/ClientReplicationDataCollector.h
+++ b/Engine/Plugins/Developer/Concert/ConcertSync/ConcertSyncClient/Source/ConcertSyncClient/Private/Replication/Processing/ClientReplicationDataCollector.h
@@ -66,6 +66,9 @@ namespace UE::ConcertSyncClient::Replication
 		 */
 		void OnObjectStreamModified(const FSoftObjectPath& Object, TArrayView<const FGuid> PutStreams);
 
+		/** Clears all currently replicated objects. */
+		void ClearReplicatedObjects();
+
 		/** Iterates every object for which there is at least one owning stream. */
 		void ForEachOwnedObject(TFunctionRef<EBreakBehavior(const FSoftObjectPath&)> Callback) const;
 		/** Writes all owning streams for ObjectPath into Paths. */
diff --git a/Engine/Plugins/Developer/Concert/ConcertSync/ConcertSyncClient/Source/ConcertSyncClient/Public/Replication/IConcertClientReplicationManager.h b/Engine/Plugins/Developer/Concert/ConcertSync/ConcertSyncClient/Source/ConcertSyncClient/Public/Replication/IConcertClientReplicationManager.h
--- a/Engine/Plugins/Developer/Concert/ConcertSync/ConcertSyncClient/Source/ConcertSyncClient/Public/Replication/IConcertClientReplicationManager.h
+++ b/Engine/Plugins/Developer/Concert/ConcertSync/ConcertSyncClient/Source/ConcertSyncClient/Public/Replication/IConcertClientReplicationManager.h
@@ -8,6 +8,7 @@
 #include "Replication/Messages/ClientQuery.h"
 #include "Replication/Messages/Handshake.h"
 #include "Replication/Messages/Muting.h"
+#include "Replication/Messages/RestoreContent.h"
 
 template<typename ResultType>
 class TFuture;
@@ -133,10 +134,10 @@ public:
 	virtual TFuture<FConcertReplication_QueryMuteState_Response> QueryMuteState(FConcertReplication_QueryMuteState_Request Request = {}) = 0;
 	TFuture<FConcertReplication_QueryMuteState_Response> QueryMuteState(TSet<FSoftObjectPath> Objects);
 
-	DECLARE_MULTICAST_DELEGATE_TwoParams(FOnPreStreamsChanged,
-		const FConcertReplication_ChangeStream_Request&,
-		const FConcertReplication_ChangeStream_Response&
-		);
+	/** Restore this client's stream content and authority to what a client had when they left. */
+	virtual TFuture<FConcertReplication_RestoreContent_Response> RestoreContent(FConcertReplication_RestoreContent_Request Request = {}) = 0;
+
+	DECLARE_MULTICAST_DELEGATE(FOnPreStreamsChanged);
 	/** Called right before the result of GetRegisteredStreams changes. */
 	virtual FOnPreStreamsChanged& OnPreStreamsChanged() = 0;
 	
@@ -144,10 +145,7 @@ public:
 	/** Called right after the result of GetRegisteredStreams has changed. */
 	virtual FOnPostStreamsChanged& OnPostStreamsChanged() = 0;
 
-	DECLARE_MULTICAST_DELEGATE_TwoParams(FOnPreAuthorityChanged,
-		const FConcertReplication_ChangeAuthority_Request&,
-		const FConcertReplication_ChangeAuthority_Response&
-		);
+	DECLARE_MULTICAST_DELEGATE(FOnPreAuthorityChanged);
 	/** Called right before GetClientOwnedObjects changes. */
 	virtual FOnPreAuthorityChanged& OnPreAuthorityChanged() = 0;
 	
diff --git a/Engine/Plugins/Developer/Concert/ConcertSync/ConcertSyncCore/Source/ConcertSyncCore/Public/Replication/Data/ClientQueriedInfo.h b/Engine/Plugins/Developer/Concert/ConcertSync/ConcertSyncCore/Source/ConcertSyncCore/Public/Replication/Data/ClientQueriedInfo.h
--- a/Engine/Plugins/Developer/Concert/ConcertSync/ConcertSyncCore/Source/ConcertSyncCore/Public/Replication/Data/ClientQueriedInfo.h
+++ b/Engine/Plugins/Developer/Concert/ConcertSync/ConcertSyncCore/Source/ConcertSyncCore/Public/Replication/Data/ClientQueriedInfo.h
@@ -2,6 +2,7 @@
 
 #pragma once
 
+#include "ObjectIds.h"
 #include "ReplicationStream.h"
 #include "ClientQueriedInfo.generated.h"
 
@@ -46,4 +47,13 @@ struct FConcertQueriedClientInfo
 	 */
 	UPROPERTY()
 	TArray<FConcertAuthorityClientInfo> Authority;
+
+	bool IsEmpty() const { return Streams.IsEmpty() && Authority.IsEmpty(); }
+	
+	/** @return Whether Object is owned */
+	bool HasAuthority(const FConcertObjectInStreamID& Object) const
+	{
+		const FConcertAuthorityClientInfo* Info = Authority.FindByPredicate([StreamId = Object.StreamId](const FConcertAuthorityClientInfo& Info){ return Info.StreamId == StreamId; });
+		return Info && Info->AuthoredObjects.Contains(Object.Object);
+	}
 };
\ No newline at end of file
diff --git a/Engine/Plugins/Developer/Concert/ConcertSync/ConcertSyncCore/Source/ConcertSyncCore/Public/Replication/Messages/RestoreContent.h b/Engine/Plugins/Developer/Concert/ConcertSync/ConcertSyncCore/Source/ConcertSyncCore/Public/Replication/Messages/RestoreContent.h
new file mode 100644
--- /dev/null
+++ b/Engine/Plugins/Developer/Concert/ConcertSync/ConcertSyncCore/Source/ConcertSyncCore/Public/Replication/Messages/RestoreContent.h
@@ -0,0 +1,134 @@
+﻿// Copyright Epic Games, Inc. All Rights Reserved.
+
+#pragma once
+
+#include "HAL/Platform.h"
+#include "Misc/Optional.h"
+#include "Replication/Data/ClientQueriedInfo.h"
+#include "Replication/Messages/SyncControl.h"
+#include "RestoreContent.generated.h"
+
+UENUM()
+enum class EConcertReplicationRestoreContentFlags : uint8
+{
+	None = 0,
+
+	/**
+	 * If set, the client's final streams & authority shall be the union of the current registered streams & authority and that of the activity.
+	 * If unset, the client's streams and authority are cleared (effectively replaced by what is stored in the activity).
+	 */
+	RestoreOnTop = 1 << 0,
+
+	/** If set, checks that there is no other client with the same display and device name. If there is, the request fails with NameConflict. */
+	ValidateUniqueClient = 1 << 1,
+
+	/** Also restore the authority the client had.  */
+	RestoreAuthority = 1 << 2,
+
+	/** The response should include the new stream and authority state. */
+	SendNewState = 1 << 3,
+
+	StreamsAndAuthority = RestoreAuthority
+};
+ENUM_CLASS_FLAGS(EConcertReplicationRestoreContentFlags);
+
+UENUM()
+enum class EConcertReplicationAuthorityRestoreMode : uint8
+{
+	// Suppose client 1 has registered (SceneComponent, { RelativeLocation }) and has authority over SceneComponent.
+	// If client 2 tries to restore (SceneComponent, { ComponentTags }), it will work in each of the below cases (client 1 does not own ComponentTags).
+	// If client 2 tries to restore (SceneComponent, { ComponentTags }), it will only work for IncludeAlreadyOwnedObjectPropertiesInStream.
+
+	// The following fields apply even when EConcertReplicationRestoreContentFlags::RestoreAuthority is not set.
+	
+	/**
+	 * If another client already has authority over a would-be restored object' restored properties, do not restore object in the client's stream.
+	 * This only affects the restore properties; if RestoreOnTop is set and there are conflicts with non-restored properties, this filter does not apply.
+	 */
+	ExcludeAlreadyOwnedObjectPropertiesFromStream,
+	/** If another client already has authority over a would-be restored object's restored properties, restore it into the client's stream (but obviously don't take authority). */
+	IncludeAlreadyOwnedObjectPropertiesInStream,
+	/** If another client already has authority over a would-be restored object's restored properties, do not restore anything. */
+	AllOrNothing
+};
+
+/**
+ * Request the sending client's stream and optionally authority to be aggregated with what a client had when they left the session.
+ *
+ * A FConcertSyncFConcertSyncReplicationActivity is generated whenever a client leaves a replication session which has
+ * EConcertSyncSessionFlags::ShouldEnableReplicationActivities flag set. This activity contains the stream content and authority of the client at that
+ * point. This request is used to restore the client's state when they last left.
+ */
+USTRUCT()
+struct FConcertReplication_RestoreContent_Request
+{
+	GENERATED_BODY()
+
+	/** Describes what and how content is to be restored. */
+	UPROPERTY()
+	EConcertReplicationRestoreContentFlags Flags = EConcertReplicationRestoreContentFlags::StreamsAndAuthority;
+
+	/** If EConcertReplicationRestoreContentFlags::RestoreAuthority is set, describes how to deal with authority conflicts. */
+	UPROPERTY()
+	EConcertReplicationAuthorityRestoreMode AuthorityRestorationMode = EConcertReplicationAuthorityRestoreMode::ExcludeAlreadyOwnedObjectPropertiesFromStream;
+
+	/**
+	 * The ID of an activity that contains a client's replication state.
+	 * The activity must be a FConcertSyncReplicationActivity with whose EventData.ActivityType == EConcertSyncReplicationActivityType::LeaveReplication.
+	 * 
+	 * If left unset, use the latest activity of matching the sending endpoint's display and device name.
+	 * This can be a different client.
+	 */
+	UPROPERTY()
+	TOptional<int64> ActivityId;
+};
+
+UENUM()
+enum class EConcertReplicationRestoreErrorCode : uint8
+{
+	/** The request timed out. */
+	Timeout,
+
+	/**
+	 * The request was successful.
+	 * The request is also successful if there was no activity to restore since the client is in the "right" state.
+	 */
+	Success,
+
+	/** Requesting client has not joined replication. */
+	Invalid,
+	/** EConcertSyncSessionFlags::ShouldEnableReplicationActivities is not set. */
+	NotSupported,
+	/** The request's ActivityId was set but did not point to an appropriate activity. */
+	NoSuchActivity,
+	/** Nothing could be restored because another client with the same display and device name already is in the session. */
+	NameConflict,
+	/** EConcertReplicationAuthorityRestoreMode::AllOrNothing was set and another client had authority over one of the would-be restore objects. */
+	AuthorityConflict
+};
+
+USTRUCT()
+struct FConcertReplication_RestoreContent_Response
+{
+	GENERATED_BODY()
+
+	UPROPERTY()
+	EConcertReplicationRestoreErrorCode ErrorCode = EConcertReplicationRestoreErrorCode::Timeout;
+
+	/**
+	 * This includes the full stream and authority content if ErrorCode == EConcertReplicationRestoreErrorCode::Success and the request had
+	 * EConcertReplicationRestoreContentFlags::SendNewState set.
+	 */
+	UPROPERTY()
+	FConcertQueriedClientInfo ClientInfo;
+	
+	/**
+	 * The full sync control the client has on the server.
+	 * Only valid if ErrorCode == EConcertReplicationRestoreErrorCode::Success.
+	 * The client does not receive a separate event for this and is expected to apply it instantly.
+	 */
+	UPROPERTY()
+	FConcertReplication_ChangeSyncControl SyncControl;
+
+	bool IsSuccess() const { return ErrorCode == EConcertReplicationRestoreErrorCode::Success; }
+};
\ No newline at end of file
diff --git a/Engine/Plugins/Developer/Concert/ConcertSync/ConcertSyncCore/Source/ConcertSyncCore/Public/Replication/Misc/ReplicationStreamUtils.h b/Engine/Plugins/Developer/Concert/ConcertSync/ConcertSyncCore/Source/ConcertSyncCore/Public/Replication/Misc/ReplicationStreamUtils.h
new file mode 100644
--- /dev/null
+++ b/Engine/Plugins/Developer/Concert/ConcertSync/ConcertSyncCore/Source/ConcertSyncCore/Public/Replication/Misc/ReplicationStreamUtils.h
@@ -0,0 +1,77 @@
+﻿// Copyright Epic Games, Inc. All Rights Reserved.
+
+#pragma once
+
+#include "Replication/Data/ObjectIds.h"
+#include "Replication/Data/ReplicationStream.h"
+
+#include "Containers/Array.h"
+#include "Containers/ArrayView.h"
+#include "Containers/ContainersFwd.h"
+#include "UObject/SoftObjectPath.h"
+
+namespace UE::ConcertSyncCore
+{
+	/** Util for finding a stream in an array of streams by id. */
+	const FConcertReplicationStream* FindStream(const TConstArrayView<FConcertReplicationStream> Streams, const FGuid& StreamId);
+	/** Util for finding a stream in an array of streams by id. */
+	FConcertReplicationStream* FindStreamEditable(const TArrayView<FConcertReplicationStream> Streams, const FGuid& StreamId);
+	
+	/** Util for finding an object in an array of streams. */
+	const FConcertReplicatedObjectInfo* FindObjectInfo(const TConstArrayView<FConcertReplicationStream> Streams, const FConcertObjectInStreamID& ObjectId);
+	/** Util for finding an object an array of streams. */
+	FConcertReplicatedObjectInfo* FindObjectInfoEditable(const TArrayView<FConcertReplicationStream> Streams, const FConcertObjectInStreamID& ObjectId);
+
+	/** Util for finding an object in a stream. */
+	const FConcertReplicatedObjectInfo* FindObjectInfo(const FConcertReplicationStream& Stream, const FSoftObjectPath& ObjectPath);
+	/** Util for finding an object in a stream. */
+	FConcertReplicatedObjectInfo* FindObjectInfoEditable(FConcertReplicationStream& Stream, const FSoftObjectPath& ObjectPath);
+	
+	/** Util for finding an object's frequency override settings in a stream in an array of streams. */
+	const FConcertObjectReplicationSettings* FindObjectFrequency(const TConstArrayView<FConcertReplicationStream> Streams, const FConcertObjectInStreamID& ObjectId);
+	/** Util for finding an object's frequency override settings in a stream in an array of streams. */
+	FConcertObjectReplicationSettings* FindObjectFrequencyEditable(const TArrayView<FConcertReplicationStream> Streams, const FConcertObjectInStreamID& ObjectId);
+}
+
+namespace UE::ConcertSyncCore
+{
+	inline const FConcertReplicationStream* FindStream(const TConstArrayView<FConcertReplicationStream> Streams, const FGuid& StreamId)
+	{
+		return Streams.FindByPredicate([&StreamId](const FConcertReplicationStream& Stream) { return Stream.BaseDescription.Identifier == StreamId; });
+	}
+	inline FConcertReplicationStream* FindStreamEditable(const TArrayView<FConcertReplicationStream> Streams, const FGuid& StreamId)
+	{
+		return Streams.FindByPredicate([&StreamId](const FConcertReplicationStream& Stream) { return Stream.BaseDescription.Identifier == StreamId; });
+	}
+	
+	inline const FConcertReplicatedObjectInfo* FindObjectInfo(const TConstArrayView<FConcertReplicationStream> Streams, const FConcertObjectInStreamID& ObjectId)
+	{
+		const FConcertReplicationStream* Stream = FindStream(Streams, ObjectId.StreamId);
+		return Stream ? Stream->BaseDescription.ReplicationMap.ReplicatedObjects.Find(ObjectId.Object) : nullptr;
+	}
+	inline FConcertReplicatedObjectInfo* FindObjectInfoEditable(const TArrayView<FConcertReplicationStream> Streams, const FConcertObjectInStreamID& ObjectId)
+	{
+		FConcertReplicationStream* Stream = FindStreamEditable(Streams, ObjectId.StreamId);
+		return Stream ? Stream->BaseDescription.ReplicationMap.ReplicatedObjects.Find(ObjectId.Object) : nullptr;
+	}
+	
+	inline const FConcertReplicatedObjectInfo* FindObjectInfo(const FConcertReplicationStream& Stream, const FSoftObjectPath& ObjectPath)
+	{
+		return Stream.BaseDescription.ReplicationMap.ReplicatedObjects.Find(ObjectPath);
+	}
+	inline FConcertReplicatedObjectInfo* FindObjectInfoEditable(FConcertReplicationStream& Stream, const FSoftObjectPath& ObjectPath)
+	{
+		return Stream.BaseDescription.ReplicationMap.ReplicatedObjects.Find(ObjectPath);
+	}
+	
+	inline const FConcertObjectReplicationSettings* FindObjectFrequency(const TConstArrayView<FConcertReplicationStream> Streams, const FConcertObjectInStreamID& ObjectId)
+	{
+		const FConcertReplicationStream* Stream = FindStream(Streams, ObjectId.StreamId);
+		return Stream ? Stream->BaseDescription.FrequencySettings.ObjectOverrides.Find(ObjectId.Object) : nullptr;
+	}
+	inline FConcertObjectReplicationSettings* FindObjectFrequencyEditable(const TArrayView<FConcertReplicationStream> Streams, const FConcertObjectInStreamID& ObjectId)
+	{
+		FConcertReplicationStream* Stream = FindStreamEditable(Streams, ObjectId.StreamId);
+		return Stream ? Stream->BaseDescription.FrequencySettings.ObjectOverrides.Find(ObjectId.Object) : nullptr;
+	}
+}
diff --git a/Engine/Plugins/Developer/Concert/ConcertSync/ConcertSyncCore/Source/ConcertSyncCore/Public/Replication/SyncControlState.h b/Engine/Plugins/Developer/Concert/ConcertSync/ConcertSyncCore/Source/ConcertSyncCore/Public/Replication/SyncControlState.h
--- a/Engine/Plugins/Developer/Concert/ConcertSync/ConcertSyncCore/Source/ConcertSyncCore/Public/Replication/SyncControlState.h
+++ b/Engine/Plugins/Developer/Concert/ConcertSync/ConcertSyncCore/Source/ConcertSyncCore/Public/Replication/SyncControlState.h
@@ -14,6 +14,7 @@
 #include <type_traits>
 
 #include "Messages/Muting.h"
+#include "Messages/RestoreContent.h"
 
 namespace UE::ConcertSyncCore::Replication
 {
@@ -66,31 +67,43 @@ namespace UE::ConcertSyncCore::Replication
 
 		/** Combines implicit and explicit sync control changes caused by authority changes. */
 		template<CObjectInStreamInvocable TOnAllowed, CObjectInStreamInvocable TOnDisallowed>
-		void AppendChanges(const FConcertReplication_ChangeAuthority_Request& Request, const FConcertReplication_ChangeAuthority_Response& Response, TOnAllowed&& OnAllowed, TOnDisallowed&& OnDisallowed);
-		void AppendChanges(const FConcertReplication_ChangeAuthority_Request& Request, const FConcertReplication_ChangeAuthority_Response& Response);
+		void AppendAuthorityChange(const FConcertReplication_ChangeAuthority_Request& Request, const FConcertReplication_ChangeAuthority_Response& Response, TOnAllowed&& OnAllowed, TOnDisallowed&& OnDisallowed);
+		void AppendAuthorityChange(const FConcertReplication_ChangeAuthority_Request& Request, const FConcertReplication_ChangeAuthority_Response& Response);
 
 		/** Combines implicit sync control changes caused by stream change. */
 		template<CObjectInStreamInvocable TOnDisallowed>
-		void AppendChanges(const FConcertReplication_ChangeStream_Request& Request, TOnDisallowed&& OnDisallowed);
-		void AppendChanges(const FConcertReplication_ChangeStream_Request& Request);
-
+		void AppendStreamChange(const FConcertReplication_ChangeStream_Request& Request, TOnDisallowed&& OnDisallowed);
+		void AppendStreamChange(const FConcertReplication_ChangeStream_Request& Request);
+		
 		/**
 		 * Combines implicit sync control changes caused by the local client sending a mute event.
 		 * @return The objects that were predictively removed by the request. Pass this back into AppendChanges once you get the response.
 		 */
 		template<CObjectInStreamInvocable TOnDisallowed>
-		FPredictedObjectRemoval AppendChanges(const FConcertReplication_ChangeMuteState_Request& Request, TOnDisallowed&& OnDisallowed);
-		FPredictedObjectRemoval AppendChanges(const FConcertReplication_ChangeMuteState_Request& Request);
-
+		FPredictedObjectRemoval PredictAndApplyMuteChanges(const FConcertReplication_ChangeMuteState_Request& Request, TOnDisallowed&& OnDisallowed);
+		FPredictedObjectRemoval PredictAndApplyMuteChanges(const FConcertReplication_ChangeMuteState_Request& Request);
 		/**
 		 * Looks at the response:
 		 * - if the change failed, reverts the predictively removed sync control
 		 * - if the change succeeded, appends the contained sync control
 		 */
 		template<CObjectInStreamInvocable TOnAllowed>
-		void AppendChanges(const FPredictedObjectRemoval& ObjectsRemovedInRequest, const FConcertReplication_ChangeMuteState_Response& Response, TOnAllowed&& OnAllowed);
-		void AppendChanges(const FPredictedObjectRemoval& ObjectsRemovedInRequest, const FConcertReplication_ChangeMuteState_Response& Response);
+		void ApplyOrRevertMuteResponse(const FPredictedObjectRemoval& ObjectsRemovedInRequest, const FConcertReplication_ChangeMuteState_Response& Response, TOnAllowed&& OnAllowed);
+		void ApplyOrRevertMuteResponse(const FPredictedObjectRemoval& ObjectsRemovedInRequest, const FConcertReplication_ChangeMuteState_Response& Response);
 		
+		/**
+		 * If the request replaces the content, removes sync control for all objects.
+		 * @return The objects that were predictively removed by the request. Pass this back into AppendChanges once you get the response.
+		 */
+		template<CObjectInStreamInvocable TOnDisallowed>
+		FPredictedObjectRemoval PredictAndApplyRestoreContentChanges(const FConcertReplication_RestoreContent_Request& Request, TOnDisallowed&& OnDisallowed);
+		/**
+		 * Looks at the response
+		 * - if the change failed, reverts the predictively removed sync control.
+		 * - if the change succeeded, appends the contained sync control.
+		 */
+		template<CObjectInStreamInvocable TOnAllowed, CObjectInStreamInvocable TOnDisallowed>
+		void ApplyOrRevertRestoreContentResponse(const FPredictedObjectRemoval& ObjectsRemovedInRequest, const FConcertReplication_RestoreContent_Response& Response, TOnAllowed&& OnAllowed, TOnDisallowed&& OnDisallowed);
 		
 		friend bool operator==(const FSyncControlState& Left, const FSyncControlState& Right)
 		{
@@ -150,7 +163,7 @@ namespace UE::ConcertSyncCore::Replication
 	}
 	
 	template <CObjectInStreamInvocable TOnAllowed, CObjectInStreamInvocable TOnDisallowed>
-	void FSyncControlState::AppendChanges(const FConcertReplication_ChangeAuthority_Request& Request, const FConcertReplication_ChangeAuthority_Response& Response, TOnAllowed&& OnAllowed, TOnDisallowed&& OnDisallowed)
+	void FSyncControlState::AppendAuthorityChange(const FConcertReplication_ChangeAuthority_Request& Request, const FConcertReplication_ChangeAuthority_Response& Response, TOnAllowed&& OnAllowed, TOnDisallowed&& OnDisallowed)
 	{
 		for (const TPair<FSoftObjectPath, FConcertStreamArray>& ImplicitChange : Request.ReleaseAuthority)
 		{
@@ -169,13 +182,13 @@ namespace UE::ConcertSyncCore::Replication
 		AppendChanges(Response.SyncControl, MoveTemp(OnAllowed), MoveTemp(OnDisallowed));
 	}
 
-	inline void FSyncControlState::AppendChanges(const FConcertReplication_ChangeAuthority_Request& Request, const FConcertReplication_ChangeAuthority_Response& Response)
+	inline void FSyncControlState::AppendAuthorityChange(const FConcertReplication_ChangeAuthority_Request& Request, const FConcertReplication_ChangeAuthority_Response& Response)
 	{
-		AppendChanges(Request, Response, [](const FConcertObjectInStreamID&){}, [](const FConcertObjectInStreamID&){});
+		AppendAuthorityChange(Request, Response, [](const FConcertObjectInStreamID&){}, [](const FConcertObjectInStreamID&){});
 	}
 	
 	template <CObjectInStreamInvocable TOnDisallowed>
-	void FSyncControlState::AppendChanges(const FConcertReplication_ChangeStream_Request& Request, TOnDisallowed&& OnDisallowed)
+	void FSyncControlState::AppendStreamChange(const FConcertReplication_ChangeStream_Request& Request, TOnDisallowed&& OnDisallowed)
 	{
 		for (auto It = AllowedObjects.CreateIterator(); It; ++It)
 		{
@@ -207,13 +220,13 @@ namespace UE::ConcertSyncCore::Replication
 		}
 	}
 
-	inline void FSyncControlState::AppendChanges(const FConcertReplication_ChangeStream_Request& Request)
+	inline void FSyncControlState::AppendStreamChange(const FConcertReplication_ChangeStream_Request& Request)
 	{
-		AppendChanges(Request, [](const FConcertObjectInStreamID&){});
+		AppendStreamChange(Request, [](const FConcertObjectInStreamID&){});
 	}
 
 	template <CObjectInStreamInvocable TOnDisallowed>
-	FSyncControlState::FPredictedObjectRemoval FSyncControlState::AppendChanges(const FConcertReplication_ChangeMuteState_Request& Request, TOnDisallowed&& OnDisallowed)
+	FSyncControlState::FPredictedObjectRemoval FSyncControlState::PredictAndApplyMuteChanges(const FConcertReplication_ChangeMuteState_Request& Request, TOnDisallowed&& OnDisallowed)
 	{
 		FPredictedObjectRemoval RemovedObjects;
 		const auto Remove = [&OnDisallowed, &RemovedObjects](TSet<FConcertObjectInStreamID>::TIterator& It)
@@ -256,13 +269,13 @@ namespace UE::ConcertSyncCore::Replication
 		return RemovedObjects;
 	}
 
-	inline FSyncControlState::FPredictedObjectRemoval FSyncControlState::AppendChanges(const FConcertReplication_ChangeMuteState_Request& Request)
+	inline FSyncControlState::FPredictedObjectRemoval FSyncControlState::PredictAndApplyMuteChanges(const FConcertReplication_ChangeMuteState_Request& Request)
 	{
-		return AppendChanges(Request, [](const FConcertObjectInStreamID&){});
+		return PredictAndApplyMuteChanges(Request, [](const FConcertObjectInStreamID&){});
 	}
 
 	template <CObjectInStreamInvocable TOnAllowed>
-	void FSyncControlState::AppendChanges(const FPredictedObjectRemoval& ObjectsRemovedInRequest, const FConcertReplication_ChangeMuteState_Response& Response, TOnAllowed&& OnAllowed)
+	void FSyncControlState::ApplyOrRevertMuteResponse(const FPredictedObjectRemoval& ObjectsRemovedInRequest, const FConcertReplication_ChangeMuteState_Response& Response, TOnAllowed&& OnAllowed)
 	{
 		if (Response.IsSuccess())
 		{
@@ -279,9 +292,50 @@ namespace UE::ConcertSyncCore::Replication
 		}
 	}
 
-	inline void FSyncControlState::AppendChanges(const FPredictedObjectRemoval& ObjectsRemovedInRequest, const FConcertReplication_ChangeMuteState_Response& Response)
+	inline void FSyncControlState::ApplyOrRevertMuteResponse(const FPredictedObjectRemoval& ObjectsRemovedInRequest, const FConcertReplication_ChangeMuteState_Response& Response)
+	{
+		ApplyOrRevertMuteResponse(ObjectsRemovedInRequest, Response, [](const FConcertObjectInStreamID&){});
+	}
+
+	template <CObjectInStreamInvocable TOnDisallowed>
+	FSyncControlState::FPredictedObjectRemoval FSyncControlState::PredictAndApplyRestoreContentChanges(
+		const FConcertReplication_RestoreContent_Request& Request,
+		TOnDisallowed&& OnDisallowed
+		)
 	{
-		AppendChanges(ObjectsRemovedInRequest, Response, [](const FConcertObjectInStreamID&){});
+		FPredictedObjectRemoval PredictedRemoval;
+		if (EnumHasAnyFlags(Request.Flags, EConcertReplicationRestoreContentFlags::RestoreOnTop))
+		{
+			return PredictedRemoval;
+		}
+
+		// The request will replace everything we have - so disallow all. 
+		for (const FConcertObjectInStreamID& Object : AllowedObjects)
+		{
+			OnDisallowed(Object);
+		}
+		PredictedRemoval.Objects = MoveTemp(AllowedObjects);
+		
+		return PredictedRemoval;
+	}
+
+	template <CObjectInStreamInvocable TOnAllowed, CObjectInStreamInvocable TOnDisallowed>
+	void FSyncControlState::ApplyOrRevertRestoreContentResponse(
+		const FPredictedObjectRemoval& ObjectsRemovedInRequest,
+		const FConcertReplication_RestoreContent_Response& Response,
+		TOnAllowed&& OnAllowed,
+		TOnDisallowed&& OnDisallowed
+		)
+	{
+		if (Response.IsSuccess())
+		{
+			AppendChanges(Response.SyncControl, OnAllowed, OnDisallowed);
+		}
+		else
+		{
+			// Do not assign MoveTemp(ObjectsRemovedInRequest.Objects) because we have been modified between!
+			AllowedObjects.Append(ObjectsRemovedInRequest.Objects);
+		}
 	}
 }
 
diff --git a/Engine/Plugins/Developer/Concert/ConcertSync/ConcertSyncServer/Source/ConcertSyncServer/Private/ConcertServerWorkspace.cpp b/Engine/Plugins/Developer/Concert/ConcertSync/ConcertSyncServer/Source/ConcertSyncServer/Private/ConcertServerWorkspace.cpp
--- a/Engine/Plugins/Developer/Concert/ConcertSync/ConcertSyncServer/Source/ConcertSyncServer/Private/ConcertServerWorkspace.cpp
+++ b/Engine/Plugins/Developer/Concert/ConcertSync/ConcertSyncServer/Source/ConcertSyncServer/Private/ConcertServerWorkspace.cpp
@@ -43,6 +43,18 @@ void FConcertServerWorkspace::ProduceClientLeaveReplicationActivity(const FGuid&
 	}
 }
 
+bool FConcertServerWorkspace::GetLastLeaveReplicationActivityByClient(const FConcertClientInfo& InClientInfo, FConcertSyncReplicationPayload_LeaveReplication& OutLeaveReplication) const
+{
+	// TODO: DO
+	return false;
+}
+
+bool FConcertServerWorkspace::GetLeaveReplicationActivityById(const int64 ActivityId, FConcertSyncReplicationPayload_LeaveReplication& OutLeaveReplication) const
+{
+	// TODO: DO
+	return false;
+}
+
 void FConcertServerWorkspace::BindSession(const TSharedRef<FConcertSyncServerLiveSession>& InLiveSession)
 {
 	check(InLiveSession->IsValidSession());
diff --git a/Engine/Plugins/Developer/Concert/ConcertSync/ConcertSyncServer/Source/ConcertSyncServer/Private/ConcertServerWorkspace.h b/Engine/Plugins/Developer/Concert/ConcertSync/ConcertSyncServer/Source/ConcertSyncServer/Private/ConcertServerWorkspace.h
--- a/Engine/Plugins/Developer/Concert/ConcertSync/ConcertSyncServer/Source/ConcertSyncServer/Private/ConcertServerWorkspace.h
+++ b/Engine/Plugins/Developer/Concert/ConcertSync/ConcertSyncServer/Source/ConcertSyncServer/Private/ConcertServerWorkspace.h
@@ -40,6 +40,8 @@ public:
 
 	//~ Begin IReplicationWorkspace Interface
 	virtual void ProduceClientLeaveReplicationActivity(const FGuid& EndpointId, const FConcertSyncReplicationPayload_LeaveReplication& EventData) override;
+	virtual bool GetLastLeaveReplicationActivityByClient(const FConcertClientInfo& InClientInfo, FConcertSyncReplicationPayload_LeaveReplication& OutLeaveReplication) const override;
+	virtual bool GetLeaveReplicationActivityById(const int64 ActivityId, FConcertSyncReplicationPayload_LeaveReplication& OutLeaveReplication) const override;
 	//~ End IReplicationWorkspace Interface
 
 private:
diff --git a/Engine/Plugins/Developer/Concert/ConcertSync/ConcertSyncServer/Source/ConcertSyncServer/Private/Replication/AuthorityManager.cpp b/Engine/Plugins/Developer/Concert/ConcertSync/ConcertSyncServer/Source/ConcertSyncServer/Private/Replication/AuthorityManager.cpp
--- a/Engine/Plugins/Developer/Concert/ConcertSync/ConcertSyncServer/Source/ConcertSyncServer/Private/Replication/AuthorityManager.cpp
+++ b/Engine/Plugins/Developer/Concert/ConcertSync/ConcertSyncServer/Source/ConcertSyncServer/Private/Replication/AuthorityManager.cpp
@@ -193,6 +193,7 @@ namespace UE::ConcertSyncServer::Replication
 		}
 	}
 
+
 	EConcertSessionResponseCode FAuthorityManager::HandleChangeAuthorityRequest(
 		const FConcertSessionContext& Context,
 		const FConcertReplication_ChangeAuthority_Request& Request,
@@ -201,23 +202,37 @@ namespace UE::ConcertSyncServer::Replication
 	{
 		// This log does two things: 1. Identify issues in unit tests / at runtime 2. Warn about possibly malicious attempts when the server runs.
 		UE_CLOG(Request.TakeAuthority.IsEmpty() && Request.ReleaseAuthority.IsEmpty(), LogConcert, Warning, TEXT("Received invalid authority request (TakeAuthority.Num() == 0 and ReleaseAuthority.Num() == 0)"));
+		
+		InternalApplyChangeAuthorityRequest(Context.SourceEndpointId, Request, Response.RejectedObjects, Response.SyncControl);
+		Response.ErrorCode = EReplicationResponseErrorCode::Handled;
+		return EConcertSessionResponseCode::Success;
+	}
+	
+	void FAuthorityManager::InternalApplyChangeAuthorityRequest(
+		const FClientId& EndpointId,
+		const FConcertReplication_ChangeAuthority_Request& Request,
+		TMap<FSoftObjectPath, FConcertStreamArray>& OutRejectedObjects,
+		FConcertReplication_ChangeSyncControl& OutChangedSyncControl,
+		bool bShouldLog
+		)
+	{
 		bool bMadeChanges = false;
 		
-		FClientAuthorityData& AuthorityData = ClientAuthorityData.FindOrAdd(Context.SourceEndpointId);
-		const FClientId& ClientId = Context.SourceEndpointId;
-		Private::ForEachReplicatedObject(Request.TakeAuthority, [this, &Response, &bMadeChanges, &AuthorityData, &ClientId](const FStreamId& StreamId, const FSoftObjectPath& ObjectPath)
+		const FClientId& ClientId = EndpointId;
+		FClientAuthorityData& AuthorityData = ClientAuthorityData.FindOrAdd(ClientId);
+		Private::ForEachReplicatedObject(Request.TakeAuthority, [this, &OutRejectedObjects, bShouldLog, &bMadeChanges, &AuthorityData, &ClientId](const FStreamId& StreamId, const FSoftObjectPath& ObjectPath)
 		{
 			const FConcertReplicatedObjectId ObjectToAuthor{ { StreamId, ObjectPath }, ClientId };
 			if (CanTakeAuthority(ObjectToAuthor))
 			{
-				UE_LOG(LogConcert, Log, TEXT("Transferred authority of %s to client %s for their stream %s"), *ObjectPath.ToString(), *ClientId.ToString(EGuidFormats::Short), *StreamId.ToString(EGuidFormats::Short));
+				UE_CLOG(bShouldLog, LogConcert, Log, TEXT("Transferred authority of %s to client %s for their stream %s"), *ObjectPath.ToString(), *ClientId.ToString(EGuidFormats::Short), *StreamId.ToString(EGuidFormats::Short));
 				AuthorityData.OwnedObjects.FindOrAdd(StreamId).Add(ObjectPath);
 				bMadeChanges = true;
 			}
 			else
 			{
-				UE_LOG(LogConcert, Log, TEXT("Rejected %s request of authority over %s in stream %s"), *ClientId.ToString(EGuidFormats::Short), *ObjectPath.ToString(), *StreamId.ToString(EGuidFormats::Short));
-				Response.RejectedObjects.FindOrAdd(ObjectPath).StreamIds.AddUnique(StreamId);
+				UE_CLOG(bShouldLog, LogConcert, Log, TEXT("Rejected %s request of authority over %s in stream %s"), *ClientId.ToString(EGuidFormats::Short), *ObjectPath.ToString(), *StreamId.ToString(EGuidFormats::Short));
+				OutRejectedObjects.FindOrAdd(ObjectPath).StreamIds.AddUnique(StreamId);
 			}
 		});
 		
@@ -239,13 +254,10 @@ namespace UE::ConcertSyncServer::Replication
 			}
 		});
 
-		Response.ErrorCode = EReplicationResponseErrorCode::Handled;
 		if (bMadeChanges && ensure(GenerateSyncControlDelegate.IsBound()))
 		{
-			Response.SyncControl = GenerateSyncControlDelegate.Execute(Context.SourceEndpointId);
+			OutChangedSyncControl = GenerateSyncControlDelegate.Execute(ClientId);
 		}
-		
-		return EConcertSessionResponseCode::Success;
 	}
 	
 	const FConcertReplicationStream* FAuthorityManager::FindClientStreamById(const FClientId& ClientId, const FStreamId& StreamId) const
diff --git a/Engine/Plugins/Developer/Concert/ConcertSync/ConcertSyncServer/Source/ConcertSyncServer/Private/Replication/AuthorityManager.h b/Engine/Plugins/Developer/Concert/ConcertSync/ConcertSyncServer/Source/ConcertSyncServer/Private/Replication/AuthorityManager.h
--- a/Engine/Plugins/Developer/Concert/ConcertSync/ConcertSyncServer/Source/ConcertSyncServer/Private/Replication/AuthorityManager.h
+++ b/Engine/Plugins/Developer/Concert/ConcertSync/ConcertSyncServer/Source/ConcertSyncServer/Private/Replication/AuthorityManager.h
@@ -10,6 +10,7 @@
 #include "Misc/EBreakBehavior.h"
 #include "Templates/Function.h"
 
+struct FConcertStreamArray;
 struct FConcertReplication_ChangeSyncControl;
 class IConcertSession;
 
@@ -78,7 +79,13 @@ namespace UE::ConcertSyncServer::Replication
 		void OnPostClientLeft(const FClientId& ClientEndpointId);
 		/** Takes away authority from the given client from the given object. */
 		void RemoveAuthority(const FConcertReplicatedObjectId& Object);
-		
+
+		/** Applies a change authority request as if EndpointId had sent it. */
+		void ApplyChangeAuthorityRequest(const FClientId& EndpointId, const FConcertReplication_ChangeAuthority_Request& Request, TMap<FSoftObjectPath, FConcertStreamArray>& OutRejectedObjects, FConcertReplication_ChangeSyncControl& OutChangedSyncControl)
+		{
+			InternalApplyChangeAuthorityRequest(EndpointId, Request, OutRejectedObjects,  OutChangedSyncControl, false);
+		}
+
 		DECLARE_DELEGATE_RetVal_OneParam(FConcertReplication_ChangeSyncControl, FGenerateSyncControl, const FGuid& ClientId);
 		/** Sets the delegate that is called to fill the sync control portion of FConcertReplication_ChangeAuthority_Response. */
 		FGenerateSyncControl& OnGenerateSyncControl() { return GenerateSyncControlDelegate; }
@@ -105,6 +112,13 @@ namespace UE::ConcertSyncServer::Replication
 			const FConcertReplication_ChangeAuthority_Request& Request,
 			FConcertReplication_ChangeAuthority_Response& Response
 			);
+		void InternalApplyChangeAuthorityRequest(
+			const FClientId& EndpointId,
+			const FConcertReplication_ChangeAuthority_Request& Request,
+			TMap<FSoftObjectPath, FConcertStreamArray>& OutRejectedObjects,
+			FConcertReplication_ChangeSyncControl& OutChangedSyncControl,
+			bool bShouldLog = true
+			);
 
 		/** Finds a stream registered with the client by its ID. */
 		const FConcertReplicationStream* FindClientStreamById(const FClientId& ClientId, const FStreamId& StreamId) const;
diff --git a/Engine/Plugins/Developer/Concert/ConcertSync/ConcertSyncServer/Source/ConcertSyncServer/Private/Replication/ConcertServerReplicationManager.cpp b/Engine/Plugins/Developer/Concert/ConcertSync/ConcertSyncServer/Source/ConcertSyncServer/Private/Replication/ConcertServerReplicationManager.cpp
--- a/Engine/Plugins/Developer/Concert/ConcertSync/ConcertSyncServer/Source/ConcertSyncServer/Private/Replication/ConcertServerReplicationManager.cpp
+++ b/Engine/Plugins/Developer/Concert/ConcertSync/ConcertSyncServer/Source/ConcertSyncServer/Private/Replication/ConcertServerReplicationManager.cpp
@@ -14,6 +14,7 @@
 #include "Replication/Messages/ClientQuery.h"
 #include "Replication/Messages/Handshake.h"
 #include "Replication/Messages/ReplicationActivity.h"
+#include "Replication/Messages/RestoreContent.h"
 #include "Replication/Processing/ObjectReplicationCache.h"
 #include "Util/JoinRequestValidation.h"
 #include "Util/LogUtils.h"
@@ -44,9 +45,10 @@ namespace UE::ConcertSyncServer::Replication
 		Session->RegisterCustomRequestHandler<FConcertReplication_Join_Request, FConcertReplication_Join_Response>(this, &FConcertServerReplicationManager::HandleJoinReplicationSessionRequest);
 		Session->RegisterCustomRequestHandler<FConcertReplication_QueryReplicationInfo_Request, FConcertReplication_QueryReplicationInfo_Response>(this, &FConcertServerReplicationManager::HandleQueryReplicationInfoRequest);
 		Session->RegisterCustomRequestHandler<FConcertReplication_ChangeStream_Request, FConcertReplication_ChangeStream_Response>(this, &FConcertServerReplicationManager::HandleChangeStreamRequest);
+		Session->RegisterCustomRequestHandler<FConcertReplication_RestoreContent_Request, FConcertReplication_RestoreContent_Response>(this, &FConcertServerReplicationManager::HandleRestoreContentRequest);
 		Session->RegisterCustomEventHandler<FConcertReplication_LeaveEvent>(this, &FConcertServerReplicationManager::HandleLeaveReplicationSessionRequest);
+		
 		Session->OnSessionClientChanged().AddRaw(this, &FConcertServerReplicationManager::OnConnectionChanged);
-
 		Session->OnTick().AddRaw(this, &FConcertServerReplicationManager::Tick);
 	}
 
diff --git a/Engine/Plugins/Developer/Concert/ConcertSync/ConcertSyncServer/Source/ConcertSyncServer/Private/Replication/ConcertServerReplicationManager.h b/Engine/Plugins/Developer/Concert/ConcertSync/ConcertSyncServer/Source/ConcertSyncServer/Private/Replication/ConcertServerReplicationManager.h
--- a/Engine/Plugins/Developer/Concert/ConcertSync/ConcertSyncServer/Source/ConcertSyncServer/Private/Replication/ConcertServerReplicationManager.h
+++ b/Engine/Plugins/Developer/Concert/ConcertSync/ConcertSyncServer/Source/ConcertSyncServer/Private/Replication/ConcertServerReplicationManager.h
@@ -16,10 +16,12 @@
 #include "SyncControlManager.h"
 
 #include "HAL/Platform.h"
+#include "Replication/Messages/RestoreContent.h"
 #include "Templates/SharedPointer.h"
 #include "Templates/Tuple.h"
 #include "Templates/UnrealTemplate.h"
 
+struct FConcertSyncReplicationPayload_LeaveReplication;
 class FConcertServerWorkspace;
 class IConcertClientReplicationBridge;
 class IConcertServerSession;
@@ -120,7 +122,20 @@ namespace UE::ConcertSyncServer::Replication
 
 		// Changing streams
 		EConcertSessionResponseCode HandleChangeStreamRequest(const FConcertSessionContext& ConcertSessionContext, const FConcertReplication_ChangeStream_Request& Request, FConcertReplication_ChangeStream_Response& Response);
-
+		/** Applies a validated change stream request. */
+		void ApplyChangeStreamRequest(const FConcertReplication_ChangeStream_Request& Request, FConcertReplicationClient& Client);
+
+		// Restoring content
+		EConcertSessionResponseCode HandleRestoreContentRequest(const FConcertSessionContext& ConcertSessionContext, const FConcertReplication_RestoreContent_Request& Request, FConcertReplication_RestoreContent_Response& Response);
+		/** @return Response code to return for the request and the payload to apply, if there is anything that should be applied. */
+		TTuple<EConcertReplicationRestoreErrorCode, TOptional<FConcertSyncReplicationPayload_LeaveReplication>> ValidateRestoreContentRequest(const FGuid& RequestingEndpointId, const FConcertReplication_RestoreContent_Request& Request) const;
+		/** Applies a validated restore content request. */
+		void ApplyRestoreContentRequest(const FGuid& RequestingEndpointId, const FConcertReplication_RestoreContent_Request& Request, const FConcertSyncReplicationPayload_LeaveReplication& DataToApply, FConcertReplication_ChangeSyncControl& ChangedSyncControl);
+		/** Restores the stream content in DataToApply to Client according to how Request specifies it. */
+		void RestoreStreamContent(const FConcertReplication_RestoreContent_Request& Request, const FConcertSyncReplicationPayload_LeaveReplication& DataToApply, FConcertReplicationClient& Client, FConcertReplication_ChangeSyncControl& OutChangedSyncControl);
+		/** Restores the authority in DataToApply to Client. */
+		void RestoreAuthority(const FConcertSyncReplicationPayload_LeaveReplication& DataToApply, const FConcertReplicationClient& Client, FConcertReplication_ChangeSyncControl& OutChangedSyncControl);
+		
 		// Leaving
 		void HandleLeaveReplicationSessionRequest(const FConcertSessionContext& ConcertSessionContext, const FConcertReplication_LeaveEvent& EventData);
 		void OnConnectionChanged(IConcertServerSession& ConcertServerSession, EConcertClientStatus ConcertClientStatus, const FConcertSessionClientInfo& ConcertSessionClientInfo);
diff --git a/Engine/Plugins/Developer/Concert/ConcertSync/ConcertSyncServer/Source/ConcertSyncServer/Private/Replication/ConcertServerReplicationManager_ChangeStream.cpp b/Engine/Plugins/Developer/Concert/ConcertSync/ConcertSyncServer/Source/ConcertSyncServer/Private/Replication/ConcertServerReplicationManager_ChangeStream.cpp
--- a/Engine/Plugins/Developer/Concert/ConcertSync/ConcertSyncServer/Source/ConcertSyncServer/Private/Replication/ConcertServerReplicationManager_ChangeStream.cpp
+++ b/Engine/Plugins/Developer/Concert/ConcertSync/ConcertSyncServer/Source/ConcertSyncServer/Private/Replication/ConcertServerReplicationManager_ChangeStream.cpp
@@ -144,7 +144,7 @@ namespace UE::ConcertSyncServer::Replication
 		false,
 		TEXT("Whether to log changes to streams.")
 		);
-	
+
 	EConcertSessionResponseCode FConcertServerReplicationManager::HandleChangeStreamRequest(
 		const FConcertSessionContext& ConcertSessionContext,
 		const FConcertReplication_ChangeStream_Request& Request,
@@ -157,29 +157,7 @@ namespace UE::ConcertSyncServer::Replication
 		const TUniquePtr<FConcertReplicationClient>* SendingClient = Clients.Find(SendingClientId);
 		if (SendingClient && Private::ShouldAcceptRequest(Request, *SendingClient->Get(), AuthorityManager, Response))
 		{
-			const TArray<FConcertReplicationStream>& StreamsBeforeChange = SendingClient->Get()->GetStreamDescriptions();
-
-			TArray<FConcertObjectInStreamID> AddedObjects;
-			ConcertSyncCore::Replication::ChangeStreamUtils::ForEachAddedObject(Request, StreamsBeforeChange, [&AddedObjects](const FConcertObjectInStreamID& Object)
-			{
-				AddedObjects.Add(Object);
-				return EBreakBehavior::Continue;
-			});
-			// If the client had authority over any objects that were removed by this request, authority must be cleaned up
-			TArray<FConcertObjectInStreamID> RemovedObjects;
-			ConcertSyncCore::Replication::ChangeStreamUtils::ForEachRemovedObject(Request, StreamsBeforeChange,
-				[this, &SendingClientId, &RemovedObjects](const FConcertObjectInStreamID& RemovedObject)
-				{
-					AuthorityManager.RemoveAuthority({ RemovedObject, SendingClientId});
-					RemovedObjects.Add(RemovedObject);
-					return EBreakBehavior::Continue;
-				});
-			
-			// This needs to happen last since all of the above use ChangeStreamUtils::ForEachRemovedObject to determine the diff of changes.
-			SendingClient->Get()->ApplyValidatedRequest(Request);
-			
-			ServerObjectCache.OnChangeStreams(SendingClientId, AddedObjects, RemovedObjects);
-			MuteManager.PostApplyStreamChange(SendingClientId, AddedObjects, RemovedObjects);
+			ApplyChangeStreamRequest(Request, *SendingClient->Get());
 		}
 		else
 		{
@@ -190,5 +168,33 @@ namespace UE::ConcertSyncServer::Replication
 		LogNetworkMessage(CVarLogStreamRequestsAndResponsesOnServer, Response, [&](){ return GetClientName(*Session, ConcertSessionContext.SourceEndpointId); });
 		return EConcertSessionResponseCode::Success;
 	}
+	
+	void FConcertServerReplicationManager::ApplyChangeStreamRequest(const FConcertReplication_ChangeStream_Request& Request, FConcertReplicationClient& Client)
+	{
+		const FGuid& SendingClientId = Client.GetClientEndpointId();
+		const TArray<FConcertReplicationStream>& StreamsBeforeChange = Client.GetStreamDescriptions();
+
+		TArray<FConcertObjectInStreamID> AddedObjects;
+		ConcertSyncCore::Replication::ChangeStreamUtils::ForEachAddedObject(Request, StreamsBeforeChange, [&AddedObjects](const FConcertObjectInStreamID& Object)
+		{
+			AddedObjects.Add(Object);
+			return EBreakBehavior::Continue;
+		});
+		// If the client had authority over any objects that were removed by this request, authority must be cleaned up
+		TArray<FConcertObjectInStreamID> RemovedObjects;
+		ConcertSyncCore::Replication::ChangeStreamUtils::ForEachRemovedObject(Request, StreamsBeforeChange,
+			[this, &SendingClientId, &RemovedObjects](const FConcertObjectInStreamID& RemovedObject)
+			{
+				AuthorityManager.RemoveAuthority({ RemovedObject, SendingClientId});
+				RemovedObjects.Add(RemovedObject);
+				return EBreakBehavior::Continue;
+			});
+			
+		// This needs to happen last since all of the above use ChangeStreamUtils::ForEachRemovedObject to determine the diff of changes.
+		Client.ApplyValidatedRequest(Request);
+			
+		ServerObjectCache.OnChangeStreams(SendingClientId, AddedObjects, RemovedObjects);
+		MuteManager.PostApplyStreamChange(SendingClientId, AddedObjects, RemovedObjects);
+	}
 }
 
diff --git a/Engine/Plugins/Developer/Concert/ConcertSync/ConcertSyncServer/Source/ConcertSyncServer/Private/Replication/ConcertServerReplicationManager_RestoreContent.cpp b/Engine/Plugins/Developer/Concert/ConcertSync/ConcertSyncServer/Source/ConcertSyncServer/Private/Replication/ConcertServerReplicationManager_RestoreContent.cpp
new file mode 100644
--- /dev/null
+++ b/Engine/Plugins/Developer/Concert/ConcertSync/ConcertSyncServer/Source/ConcertSyncServer/Private/Replication/ConcertServerReplicationManager_RestoreContent.cpp
@@ -0,0 +1,351 @@
+﻿// Copyright Epic Games, Inc. All Rights Reserved.
+
+#include "ConcertLogGlobal.h"
+#include "ConcertServerReplicationManager.h"
+#include "ConcertServer/Private/ConcertServerSession.h"
+#include "Replication/IReplicationWorkspace.h"
+#include "Replication/Messages/ChangeAuthority.h"
+#include "Replication/Messages/ChangeStream.h"
+#include "Replication/Messages/ReplicationActivity.h"
+#include "Replication/Misc/ReplicationStreamUtils.h"
+
+namespace UE::ConcertSyncServer::Replication
+{
+	namespace RestoreContentPrivate
+	{
+		/** @return Whether to consider the client infos equal for the purposes of restoring content */
+		static bool EqualsClientInfo(const FConcertClientInfo& Left, const FConcertClientInfo& Right)
+		{
+			return Left.DisplayName == Right.DisplayName && Left.DeviceName == Right.DeviceName;
+		}
+
+		/** Checks whether the request should fail according to EConcertReplicationRestoreContentFlags::ValidateUniqueClient. */
+		static bool HasNameConflict(
+			const FConcertSessionClientInfo& SenderInfo,
+			const FConcertReplication_RestoreContent_Request& Request,
+			const IConcertServerSession& Session
+			)
+		{
+			if (!EnumHasAnyFlags(Request.Flags, EConcertReplicationRestoreContentFlags::ValidateUniqueClient))
+			{
+				return false;
+			}
+			
+			for (const FGuid& SessionId : Session.GetSessionClientEndpointIds())
+			{
+				FConcertSessionClientInfo OtherInfo;
+				const bool bFoundClient = Session.FindSessionClient(SessionId, OtherInfo);
+
+				if (ensure(bFoundClient)
+					&& OtherInfo.ClientEndpointId != SenderInfo.ClientEndpointId
+					&& EqualsClientInfo(SenderInfo.ClientInfo, OtherInfo.ClientInfo))
+				{
+					return true;
+				}
+			}
+			return false;
+		}
+
+		/** Checks whether the request should fail according to EConcertReplicationAuthorityRestoreMode::AllOrNothing. */
+		static bool HasAuthorityConflict(
+			const FConcertReplication_RestoreContent_Request& Request,
+			const FConcertSyncReplicationPayload_LeaveReplication& ClientState,
+			const FAuthorityManager& AuthorityManager
+			)
+		{
+			if (Request.AuthorityRestorationMode != EConcertReplicationAuthorityRestoreMode::AllOrNothing)
+			{
+				// Request wants to fail irrelevant on whether EConcertReplicationRestoreContentFlags::RestoreAuthority is set
+				return false;
+			}
+
+			for (const FConcertObjectInStreamID& PreviouslyOwnedObject : ClientState.OwnedObjects)
+			{
+				const FConcertReplicationStream* Stream = ClientState.Streams.FindByPredicate([StreamId = PreviouslyOwnedObject.StreamId](const FConcertReplicationStream& Stream)
+				{
+					return Stream.BaseDescription.Identifier == StreamId;
+				});
+
+				const FConcertReplicatedObjectInfo* ObjectInfo = Stream ? Stream->BaseDescription.ReplicationMap.ReplicatedObjects.Find(PreviouslyOwnedObject.Object) : nullptr;
+				if (!ensure(ObjectInfo))
+				{
+					UE_LOG(LogConcert, Error, TEXT("Saved activity state violates invariant. Object %s is owned but does not appear in stream"), *PreviouslyOwnedObject.ToString());
+					continue;
+				}
+
+				const FConcertReplicatedObjectId ObjectId{ { PreviouslyOwnedObject }, Stream->BaseDescription.Identifier };
+				const bool bHasConflict = AuthorityManager.EnumerateAuthorityConflicts(ObjectId, &ObjectInfo->PropertySelection) == FAuthorityManager::EAuthorityResult::Conflict;
+				if (bHasConflict)
+				{
+					return true;
+				}
+			}
+
+			return false;
+		}
+
+		/** Fills OutChangedSyncControl with all objects the Client currently has sync control over. */
+		static void BuildRemovedSyncControl(const FConcertReplicationClient& Client, const FSyncControlManager& SyncControlManager, FConcertReplication_ChangeSyncControl& OutChangedSyncControl)
+		{
+			for (const FConcertReplicationStream& ExistingStream : Client.GetStreamDescriptions())
+			{
+				for (const TPair<FSoftObjectPath, FConcertReplicatedObjectInfo>& Object : ExistingStream.BaseDescription.ReplicationMap.ReplicatedObjects)
+				{
+					const FConcertObjectInStreamID ObjectId { ExistingStream.BaseDescription.Identifier, Object.Key };
+					const FConcertReplicatedObjectId ReplicatedObjectId { { ObjectId }, Client.GetClientEndpointId() };
+					if (SyncControlManager.HasSyncControl(ReplicatedObjectId))
+					{
+						OutChangedSyncControl.NewControlStates.Add(ObjectId, false);
+					}
+				}
+			}
+		}
+
+		/** Adds ObjectId either to StreamsToAdd or ObjectsToPut depending on whether ClientStreams already contains it. */
+		static void AddObjectToRequest(
+			FConcertReplication_ChangeStream_Request& ChangeRequest,
+			const TArray<FConcertReplicationStream>& ClientStreams,
+			const FConcertObjectInStreamID& ObjectId,
+			const FConcertReplicatedObjectInfo& ObjectInfo,
+			const FConcertReplicationStream& RestoredStream
+			)
+		{
+			const FConcertObjectReplicationSettings* ObjectFrequencyOverride = RestoredStream.BaseDescription.FrequencySettings.ObjectOverrides.Find(ObjectId.Object);
+			if (const FConcertReplicationStream* ExistingStream = ConcertSyncCore::FindStream(ClientStreams, ObjectId.StreamId))
+			{
+				FConcertReplication_ChangeStream_PutObject& PutObject = ChangeRequest.ObjectsToPut.Add(ObjectId);
+				PutObject.Properties = ObjectInfo.PropertySelection;
+				PutObject.ClassPath = ObjectInfo.ClassPath;
+				if (const FConcertReplicatedObjectInfo* ExistingObjectInfo = ConcertSyncCore::FindObjectInfo(*ExistingStream, ObjectId.Object))
+				{
+					PutObject.Properties.ReplicatedProperties.Append(ExistingObjectInfo->PropertySelection.ReplicatedProperties);
+				}
+				
+				if (ObjectFrequencyOverride)
+				{
+					ChangeRequest.FrequencyChanges.FindOrAdd(ObjectId.StreamId).OverridesToAdd.Add(ObjectId.Object) = *ObjectFrequencyOverride;
+				}
+				return;
+			}
+				
+			FConcertReplicationStream* AddedStream = ConcertSyncCore::FindStreamEditable(ChangeRequest.StreamsToAdd, ObjectId.StreamId);
+			if (!AddedStream)
+			{
+				const int32 NewStreamIndex = ChangeRequest.StreamsToAdd.Emplace();
+				AddedStream = &ChangeRequest.StreamsToAdd[NewStreamIndex];
+				AddedStream->BaseDescription.Identifier = ObjectId.StreamId;
+				AddedStream->BaseDescription.FrequencySettings.Defaults = RestoredStream.BaseDescription.FrequencySettings.Defaults;
+			}
+				
+			AddedStream->BaseDescription.ReplicationMap.ReplicatedObjects.Add(ObjectId.Object, ObjectInfo);
+			if (ObjectFrequencyOverride)
+			{
+				AddedStream->BaseDescription.FrequencySettings.ObjectOverrides.Add(ObjectId.Object, *ObjectFrequencyOverride);
+			}
+		}
+
+		/** Lists every objects that EndpointId can take authority over. */
+		template<typename TLambda> requires std::is_invocable_v<TLambda, const FConcertObjectInStreamID&>
+		static void ForEachUnownedObject(const FGuid& EndpointId, const TArray<FConcertObjectInStreamID>& Objects, const FAuthorityManager& AuthorityManager, TLambda&& Callback)
+		{
+			for (const FConcertObjectInStreamID& ObjectId : Objects)
+			{
+				const FConcertReplicatedObjectId ReplicatedObjectId { { ObjectId }, EndpointId };
+				const bool bIsAllowed = AuthorityManager.EnumerateAuthorityConflicts(ReplicatedObjectId) == FAuthorityManager::EAuthorityResult::Allowed;
+				if (bIsAllowed)
+				{
+					Callback(ObjectId);
+				}
+			}
+		}
+
+		static void FillResponseWithClientState(
+			const FConcertReplication_RestoreContent_Request& Request,
+			FConcertReplication_RestoreContent_Response& Response,
+			const FConcertReplicationClient& Client,
+			const FAuthorityManager& AuthorityManager
+			)
+		{
+			if (!EnumHasAnyFlags(Request.Flags, EConcertReplicationRestoreContentFlags::SendNewState))
+			{
+				return;
+			}
+			
+			for (const FConcertReplicationStream& Stream : Client.GetStreamDescriptions())
+			{
+				const FGuid& StreamId = Stream.BaseDescription.Identifier;
+				Response.ClientInfo.Streams.Add(Stream.BaseDescription);
+				AuthorityManager.EnumerateAuthority(Client.GetClientEndpointId(), Stream.BaseDescription.Identifier,
+					[&Response, &StreamId](const FSoftObjectPath& OwnedObject)
+					{
+						TArray<FConcertAuthorityClientInfo>& Authorities = Response.ClientInfo.Authority;
+						const int32 Index = Authorities.IndexOfByPredicate([&StreamId](const FConcertAuthorityClientInfo& Info)
+						{
+							return Info.StreamId == StreamId;
+						});
+						FConcertAuthorityClientInfo& ClientInfo = Authorities.IsValidIndex(Index)
+							? Authorities[Index]
+							: [&StreamId, &Authorities]() -> FConcertAuthorityClientInfo&
+							{
+								const int32 NewIndex = Authorities.Add(FConcertAuthorityClientInfo{ StreamId });
+								return Authorities[NewIndex];
+							}();
+						
+						ClientInfo.AuthoredObjects.AddUnique(OwnedObject);
+						return EBreakBehavior::Continue;
+					});
+			}
+		}
+	}
+	
+	EConcertSessionResponseCode FConcertServerReplicationManager::HandleRestoreContentRequest(
+		const FConcertSessionContext& ConcertSessionContext,
+		const FConcertReplication_RestoreContent_Request& Request,
+		FConcertReplication_RestoreContent_Response& Response)
+	{
+		const auto[ErrorCode, LeaveReplicationPayload] = ValidateRestoreContentRequest(ConcertSessionContext.SourceEndpointId, Request);
+		if (LeaveReplicationPayload)
+		{
+			ApplyRestoreContentRequest(ConcertSessionContext.SourceEndpointId, Request, *LeaveReplicationPayload, Response.SyncControl);
+			RestoreContentPrivate::FillResponseWithClientState(Request, Response, *Clients[ConcertSessionContext.SourceEndpointId], AuthorityManager);
+		}
+		
+		Response.ErrorCode = ErrorCode;
+		return EConcertSessionResponseCode::Success;
+	}
+
+	TTuple<EConcertReplicationRestoreErrorCode, TOptional<FConcertSyncReplicationPayload_LeaveReplication>> FConcertServerReplicationManager::ValidateRestoreContentRequest(
+		const FGuid& RequestingEndpointId,
+		const FConcertReplication_RestoreContent_Request& Request
+		) const
+	{
+		if (!Clients.Contains(RequestingEndpointId))
+		{
+			return MakeTuple(EConcertReplicationRestoreErrorCode::Invalid, TOptional<FConcertSyncReplicationPayload_LeaveReplication>{});
+		}
+		
+		if (!EnumHasAnyFlags(SessionFlags, EConcertSyncSessionFlags::ShouldEnableReplicationActivities))
+		{
+			return MakeTuple(EConcertReplicationRestoreErrorCode::NotSupported, TOptional<FConcertSyncReplicationPayload_LeaveReplication>{});
+		}
+		
+		FConcertSessionClientInfo ClientInfo;
+		const bool bFoundClient = Session->FindSessionClient(RequestingEndpointId, ClientInfo);
+		check(bFoundClient);
+
+		// Does the request want to be skipped if there is another client?
+		if (RestoreContentPrivate::HasNameConflict(ClientInfo, Request, *Session))
+		{
+			return MakeTuple(EConcertReplicationRestoreErrorCode::NameConflict, TOptional<FConcertSyncReplicationPayload_LeaveReplication>{});
+		}
+
+		// Is there even anything to restore to?
+		FConcertSyncReplicationPayload_LeaveReplication OldClientState;
+		const bool bOverrideActivity = Request.ActivityId.IsSet();
+		const bool bFound = bOverrideActivity
+			? ServerWorkspace.GetLeaveReplicationActivityById(*Request.ActivityId, OldClientState)
+			: ServerWorkspace.GetLastLeaveReplicationActivityByClient(ClientInfo.ClientInfo, OldClientState);
+		if (!bFound)
+		{
+			return MakeTuple(
+				bOverrideActivity ? EConcertReplicationRestoreErrorCode::NoSuchActivity : EConcertReplicationRestoreErrorCode::Success,
+				TOptional<FConcertSyncReplicationPayload_LeaveReplication>{}
+				);
+		}
+
+		// Does the request want to fail if another client already has authority over any of the objects?
+		if (RestoreContentPrivate::HasAuthorityConflict(Request, OldClientState, AuthorityManager))
+		{
+			return MakeTuple(EConcertReplicationRestoreErrorCode::AuthorityConflict, TOptional<FConcertSyncReplicationPayload_LeaveReplication>{});
+		}
+
+		return MakeTuple(EConcertReplicationRestoreErrorCode::Success, TOptional(OldClientState));
+	}
+
+	void FConcertServerReplicationManager::ApplyRestoreContentRequest(
+		const FGuid& RequestingEndpointId,
+		const FConcertReplication_RestoreContent_Request& Request,
+		const FConcertSyncReplicationPayload_LeaveReplication& DataToApply, 
+		FConcertReplication_ChangeSyncControl& ChangedSyncControl
+		)
+	{
+		FConcertReplicationClient& Client = *Clients[RequestingEndpointId];
+		
+		// After this first step, SyncControl will hold only the objects that were removed thus far.
+		RestoreStreamContent(Request, DataToApply, Client, ChangedSyncControl);
+		
+		if (EnumHasAnyFlags(Request.Flags, EConcertReplicationRestoreContentFlags::RestoreAuthority))
+		{
+			// This adds the objects the client receives sync control for
+			RestoreAuthority(DataToApply, Client, ChangedSyncControl);
+		}
+	}
+	
+	void FConcertServerReplicationManager::RestoreStreamContent(
+		const FConcertReplication_RestoreContent_Request& Request,
+		const FConcertSyncReplicationPayload_LeaveReplication& DataToApply,
+		FConcertReplicationClient& Client,
+		FConcertReplication_ChangeSyncControl& OutChangedSyncControl
+		)
+	{
+		const bool bClearedContent = !EnumHasAnyFlags(Request.Flags, EConcertReplicationRestoreContentFlags::RestoreOnTop); 
+		if (bClearedContent)
+		{
+			// We'll act as if all objects will remove sync control. In RestoreAuthority, we'll fix it back up.
+			RestoreContentPrivate::BuildRemovedSyncControl(Client, SyncControlManager, OutChangedSyncControl);
+			
+			FConcertReplication_ChangeStream_Request ClearRequest;
+			Algo::Transform(Client.GetStreamDescriptions(), ClearRequest.StreamsToRemove, [](const FConcertReplicationStream& Stream){ return Stream.BaseDescription.Identifier; });
+			ApplyChangeStreamRequest(ClearRequest, Client);
+		}
+
+		FConcertReplication_ChangeStream_Request ChangeRequest;
+		for (const FConcertReplicationStream& Stream : DataToApply.Streams)
+		{
+			const FGuid& StreamId = Stream.BaseDescription.Identifier;
+			for (const TPair<FSoftObjectPath, FConcertReplicatedObjectInfo>& Pair : Stream.BaseDescription.ReplicationMap.ReplicatedObjects)
+			{
+				const FConcertReplicatedObjectId ObjectId{ { StreamId, Pair.Key }, Client.GetClientEndpointId() };
+
+				// For the IncludeAlreadyOwnedObjectPropertiesInStream case, we'll add to the stream regardless of other clients' authority states.
+				// For the AllOrNothing case, ValidateRestoreContentRequest has already validated there won't be any conflicts.
+				const bool bShouldSkip = Request.AuthorityRestorationMode == EConcertReplicationAuthorityRestoreMode::ExcludeAlreadyOwnedObjectPropertiesFromStream
+					&& AuthorityManager.EnumerateAuthorityConflicts(ObjectId, &Pair.Value.PropertySelection) == FAuthorityManager::EAuthorityResult::Conflict;
+				if (!bShouldSkip)
+				{
+					RestoreContentPrivate::AddObjectToRequest(ChangeRequest, Client.GetStreamDescriptions(), ObjectId, Pair.Value, Stream);
+				}
+			}
+
+			// The stream's frequency settings also need to be restored
+			const FConcertReplicationStream* PreExistingStream = ConcertSyncCore::FindStream(Client.GetStreamDescriptions(), StreamId);
+			const FConcertObjectReplicationSettings& RestoreFrequencyDefaults = Stream.BaseDescription.FrequencySettings.Defaults;
+			if (PreExistingStream && PreExistingStream->BaseDescription.FrequencySettings.Defaults != RestoreFrequencyDefaults)
+			{
+				ChangeRequest.FrequencyChanges.FindOrAdd(StreamId).NewDefaults = RestoreFrequencyDefaults; 
+			}
+		}
+		ApplyChangeStreamRequest(ChangeRequest, Client);
+	}
+
+	void FConcertServerReplicationManager::RestoreAuthority(
+		const FConcertSyncReplicationPayload_LeaveReplication& DataToApply,
+		const FConcertReplicationClient& Client,
+		FConcertReplication_ChangeSyncControl& OutChangedSyncControl
+		)
+	{
+		// Assign every object to Client that would not cause any conflict
+		FConcertReplication_ChangeAuthority_Request Request;
+		RestoreContentPrivate::ForEachUnownedObject(Client.GetClientEndpointId(), DataToApply.OwnedObjects, AuthorityManager,
+			[&Request](const FConcertObjectInStreamID& AllowedObject)
+			{
+				Request.TakeAuthority.FindOrAdd(AllowedObject.Object).StreamIds.AddUnique(AllowedObject.StreamId);
+			});
+
+		// So far OutChangedSyncControl was filled by RestoreStreamContent: it contains only objects losing sync control.
+		// Some of those objects may now regain sync control and will be overriden.
+		TMap<FSoftObjectPath, FConcertStreamArray> RejectedObjects;
+		AuthorityManager.ApplyChangeAuthorityRequest(Client.GetClientEndpointId(), Request, RejectedObjects, OutChangedSyncControl);
+		check(RejectedObjects.IsEmpty());
+	}
+}
\ No newline at end of file
diff --git a/Engine/Plugins/Developer/Concert/ConcertSync/ConcertSyncServer/Source/ConcertSyncServer/Public/Replication/IReplicationWorkspace.h b/Engine/Plugins/Developer/Concert/ConcertSync/ConcertSyncServer/Source/ConcertSyncServer/Public/Replication/IReplicationWorkspace.h
--- a/Engine/Plugins/Developer/Concert/ConcertSync/ConcertSyncServer/Source/ConcertSyncServer/Public/Replication/IReplicationWorkspace.h
+++ b/Engine/Plugins/Developer/Concert/ConcertSync/ConcertSyncServer/Source/ConcertSyncServer/Public/Replication/IReplicationWorkspace.h
@@ -2,6 +2,9 @@
 
 #pragma once
 
+#include "HAL/Platform.h"
+
+struct FConcertClientInfo;
 struct FConcertSyncReplicationPayload_LeaveReplication;
 struct FGuid;
 
@@ -18,6 +21,20 @@ namespace UE::ConcertSyncServer::Replication
 		/** Creates a replication activity for the client leaving replication if the session has the EConcertSyncSessionFlags::ShouldEnableReplicationActivities flag. */
 		virtual void ProduceClientLeaveReplicationActivity(const FGuid& EndpointId, const FConcertSyncReplicationPayload_LeaveReplication& EventData) = 0;
 
+		/**
+		 * Gets the last replication leave activity associated for a given client.
+		 * 
+		 * As endpoint IDs change every time a client join a session, the look up is done by client display name.
+		 * If multiple machines joined with the same display name, the tie is broken by also using the device name.
+		 * 
+		 * @param InClientInfo Info about the client for which to get the activity.
+		 * @param OutLeaveReplication The activity, if present
+		 * @return Whether OutLeaveReplication contains a valid result.
+		 */
+		virtual bool GetLastLeaveReplicationActivityByClient(const FConcertClientInfo& InClientInfo, FConcertSyncReplicationPayload_LeaveReplication& OutLeaveReplication) const = 0;
+		/** Gets the replication leave activity with ActivityId. */
+		virtual bool GetLeaveReplicationActivityById(const int64 ActivityId, FConcertSyncReplicationPayload_LeaveReplication& OutLeaveReplication) const = 0;
+
 		virtual ~IReplicationWorkspace() = default;
 	};
 }
diff --git a/Engine/Plugins/Developer/Concert/ConcertSync/ConcertSyncTest/Source/ConcertSyncTest/Private/Replication/ActivityProduction.spec.cpp b/Engine/Plugins/Developer/Concert/ConcertSync/ConcertSyncTest/Source/ConcertSyncTest/Private/Replication/ActivityProduction.spec.cpp
deleted file mode 100644
diff --git a/Engine/Plugins/Developer/Concert/ConcertSync/ConcertSyncTest/Source/ConcertSyncTest/Private/Replication/Components/SyncControlState.spec.cpp b/Engine/Plugins/Developer/Concert/ConcertSync/ConcertSyncTest/Source/ConcertSyncTest/Private/Replication/Components/SyncControlState.spec.cpp
--- a/Engine/Plugins/Developer/Concert/ConcertSync/ConcertSyncTest/Source/ConcertSyncTest/Private/Replication/Components/SyncControlState.spec.cpp
+++ b/Engine/Plugins/Developer/Concert/ConcertSync/ConcertSyncTest/Source/ConcertSyncTest/Private/Replication/Components/SyncControlState.spec.cpp
@@ -42,7 +42,7 @@ namespace UE::ConcertSyncTests::Replication::UI
 				const FConcertReplication_ChangeAuthority_Request Request { .ReleaseAuthority = { { ObjectReplicator->TestObject, {{ SenderStreamId }} } } };
 				const FConcertReplication_ChangeAuthority_Response Response;
 				TArray<FConcertObjectInStreamID> RemovedObjects;
-				SyncControl.AppendChanges(
+				SyncControl.AppendAuthorityChange(
 					Request,
 					Response,
 					[this](const FConcertObjectInStreamID&){ AddError(TEXT("No object should be added")); },
@@ -61,7 +61,7 @@ namespace UE::ConcertSyncTests::Replication::UI
 				const FConcertReplication_ChangeAuthority_Request Request { .TakeAuthority = {{ ObjectReplicator->TestObject, {{ SenderStreamId }}}}};
 				const FConcertReplication_ChangeAuthority_Response Response { .SyncControl = {{{ ObjectId, true }}}};
 				TArray<FConcertObjectInStreamID> AddedObjects;
-				SyncControl.AppendChanges(
+				SyncControl.AppendAuthorityChange(
 					Request,
 					Response,
 					[&AddedObjects](const FConcertObjectInStreamID& Object){ AddedObjects.Add(Object); },
@@ -79,7 +79,7 @@ namespace UE::ConcertSyncTests::Replication::UI
 				FConcertReplication_ChangeStream_Request StreamChange;
 				StreamChange.ObjectsToRemove.Add({ SenderStreamId, ObjectReplicator->TestObject });
 				TArray<FConcertObjectInStreamID> RemovedObjects;
-				SyncControl.AppendChanges(StreamChange, [&RemovedObjects](const FConcertObjectInStreamID& Object){ RemovedObjects.Add(Object); });
+				SyncControl.AppendStreamChange(StreamChange, [&RemovedObjects](const FConcertObjectInStreamID& Object){ RemovedObjects.Add(Object); });
 				
 				TestTrue(TEXT("Removed TestObject"), RemovedObjects.Contains(FConcertObjectInStreamID{ SenderStreamId, ObjectReplicator->TestObject }));
 				TestEqual(TEXT("Removed exactly 1 object"), RemovedObjects.Num(), 1);
@@ -93,7 +93,7 @@ namespace UE::ConcertSyncTests::Replication::UI
 					
 					FConcertReplication_ChangeMuteState_Request Request;
 					Request.ObjectsToMute.Add(ObjectReplicator->TestObject, { EConcertReplicationMuteOption::OnlyObject });
-					const FSyncControlState::FPredictedObjectRemoval Removal = SyncControl.AppendChanges(Request);
+					const FSyncControlState::FPredictedObjectRemoval Removal = SyncControl.PredictAndApplyMuteChanges(Request);
 					
 					TestEqual(TEXT("SyncControl.Num() == 0"), SyncControl.Num(), 0);
 				});
@@ -104,12 +104,12 @@ namespace UE::ConcertSyncTests::Replication::UI
 					
 					FConcertReplication_ChangeMuteState_Request Request;
 					Request.ObjectsToUnmute.Add(ObjectReplicator->TestObject, { EConcertReplicationMuteOption::OnlyObject });
-					const FSyncControlState::FPredictedObjectRemoval Removal = SyncControl.AppendChanges(Request);
+					const FSyncControlState::FPredictedObjectRemoval Removal = SyncControl.PredictAndApplyMuteChanges(Request);
 					TestEqual(TEXT("SyncControl.Num() == 0"), SyncControl.Num(), 0);
 
 					FConcertReplication_ChangeMuteState_Response Response { EConcertReplicationMuteErrorCode::Accepted };
 					Response.SyncControl.NewControlStates.Add({ SenderStreamId, ObjectReplicator->TestObject }, true);
-					SyncControl.AppendChanges(Removal, Response);
+					SyncControl.ApplyOrRevertMuteResponse(Removal, Response);
 					
 					TestEqual(TEXT("SyncControl.Num() == 1"), SyncControl.Num(), 1);
 					TestTrue(TEXT("Contains TestObject"), SyncControl.IsObjectAllowed({ SenderStreamId, ObjectReplicator->TestObject }));
@@ -121,10 +121,10 @@ namespace UE::ConcertSyncTests::Replication::UI
 					
 					FConcertReplication_ChangeMuteState_Request Request;
 					Request.ObjectsToMute.Add(ObjectReplicator->TestObject, { EConcertReplicationMuteOption::OnlyObject });
-					const FSyncControlState::FPredictedObjectRemoval Removal = SyncControl.AppendChanges(Request);
+					const FSyncControlState::FPredictedObjectRemoval Removal = SyncControl.PredictAndApplyMuteChanges(Request);
 
 					const FConcertReplication_ChangeMuteState_Response Response { EConcertReplicationMuteErrorCode::Timeout };
-					SyncControl.AppendChanges(Removal, Response);
+					SyncControl.ApplyOrRevertMuteResponse(Removal, Response);
 					
 					TestEqual(TEXT("SyncControl.Num() == 1"), SyncControl.Num(), 1);
 					TestTrue(TEXT("Contains TestObject"), SyncControl.IsObjectAllowed({ SenderStreamId, ObjectReplicator->TestObject }));
@@ -140,7 +140,7 @@ namespace UE::ConcertSyncTests::Replication::UI
 					
 					FConcertReplication_ChangeMuteState_Request Request;
 					Request.ObjectsToMute.Add(ObjectReplicator->TestObject, { EConcertReplicationMuteOption::ObjectAndSubobjects });
-					const FSyncControlState::FPredictedObjectRemoval Removal = SyncControl.AppendChanges(Request);
+					const FSyncControlState::FPredictedObjectRemoval Removal = SyncControl.PredictAndApplyMuteChanges(Request);
 					
 					TestEqual(TEXT("SyncControl.Num() == 0"), SyncControl.Num(), 0);
 				});
@@ -151,7 +151,7 @@ namespace UE::ConcertSyncTests::Replication::UI
 
 					FConcertReplication_ChangeMuteState_Request Request;
 					Request.ObjectsToUnmute.Add(ObjectReplicator->TestObject, { EConcertReplicationMuteOption::ObjectAndSubobjects });
-					const FSyncControlState::FPredictedObjectRemoval Removal = SyncControl.AppendChanges(Request);
+					const FSyncControlState::FPredictedObjectRemoval Removal = SyncControl.PredictAndApplyMuteChanges(Request);
 					
 					FConcertReplication_ChangeMuteState_Response Response { EConcertReplicationMuteErrorCode::Accepted };
 					Response.SyncControl.NewControlStates =
@@ -159,7 +159,7 @@ namespace UE::ConcertSyncTests::Replication::UI
 						{ { SenderStreamId, ObjectReplicator->TestObject }, true },
 						{ { SenderStreamId, SubobjectReplicator->TestObject }, true },
 					};
-					SyncControl.AppendChanges(Removal, Response);
+					SyncControl.ApplyOrRevertMuteResponse(Removal, Response);
 					
 					TestEqual(TEXT("SyncControl.Num() == 2"), SyncControl.Num(), 2);
 					TestTrue(TEXT("Contains root"), SyncControl.IsObjectAllowed({ SenderStreamId, ObjectReplicator->TestObject }));
diff --git a/Engine/Plugins/Developer/Concert/ConcertSync/ConcertSyncTest/Source/ConcertSyncTest/Private/Replication/Restore/ActivityProduction.spec.cpp b/Engine/Plugins/Developer/Concert/ConcertSync/ConcertSyncTest/Source/ConcertSyncTest/Private/Replication/Restore/ActivityProduction.spec.cpp
new file mode 100644
--- /dev/null
+++ b/Engine/Plugins/Developer/Concert/ConcertSync/ConcertSyncTest/Source/ConcertSyncTest/Private/Replication/Restore/ActivityProduction.spec.cpp
@@ -0,0 +1,123 @@
+﻿// Copyright Epic Games, Inc. All Rights Reserved.
+
+#include "Replication/Util/Mocks/ReplicationWorkspaceCallInterceptorMock.h"
+#include "Replication/Util/Spec/ReplicationClient.h"
+#include "Replication/Util/Spec/ReplicationServer.h"
+#include "Replication/Util/Spec/ObjectTestReplicator.h"
+
+#include "Misc/AutomationTest.h"
+#include "Replication/Messages/ReplicationActivity.h"
+
+namespace UE::ConcertSyncTests::Replication
+{
+	BEGIN_DEFINE_SPEC(FLeaveReplicationActivitySpec, "Editor.Concert.Replication.RestoreContent.Activity", EAutomationTestFlags::EditorContext | EAutomationTestFlags::EngineFilter)
+		/** Detects calls into the workspace */
+		TSharedPtr<FReplicationWorkspaceCallInterceptorMock> WorkspaceMock;
+	
+		TUniquePtr<FObjectTestReplicator> ObjectReplicator;
+		TUniquePtr<FReplicationServer> Server;
+		FReplicationClient* Client = nullptr;
+	
+		FGuid SenderStreamId = FGuid::NewGuid();
+
+		void ValidateLeaveActivity()
+		{
+			if (WorkspaceMock->LastCall_ProduceClientLeaveReplicationActivity)
+			{
+				const auto[EndpointId, EventData] = *WorkspaceMock->LastCall_ProduceClientLeaveReplicationActivity;
+				TestEqual(TEXT("EndpointId"), EndpointId, Client->GetEndpointId());
+				
+				TestEqual(TEXT("Streams.Num() == 1"), EventData.Streams.Num(), 1);
+				if (EventData.Streams.Num() == 1)
+				{
+					const FConcertBaseStreamInfo& StreamInfo = EventData.Streams[0].BaseDescription;
+					TestEqual(TEXT("StreamId"), StreamInfo.Identifier, SenderStreamId);
+					
+					const FConcertReplicatedObjectInfo* ObjectInfo = StreamInfo.ReplicationMap.ReplicatedObjects.Find(ObjectReplicator->TestObject);
+					if (ObjectInfo)
+					{
+						TestEqual(TEXT("ReplicatedObjects.Num() == 1"),StreamInfo.ReplicationMap.ReplicatedObjects.Num(), 1);
+						TestTrue(TEXT("Class correct"), ObjectInfo->ClassPath == ObjectReplicator->TestObject->GetClass());
+						
+						const ConcertSyncClient::Replication::FJoinReplicatedSessionArgs SenderArgs = ObjectReplicator->CreateSenderArgs(SenderStreamId);
+						const FConcertPropertySelection& ExpectedProperties = SenderArgs.Streams[0].BaseDescription.ReplicationMap.ReplicatedObjects[ObjectReplicator->TestObject].PropertySelection;
+						TestEqual(TEXT("Properties equal"), ObjectInfo->PropertySelection, ExpectedProperties);
+					}
+					else
+					{
+						AddError(TEXT("Replicated object not listed in stream"));
+					}
+				}
+				
+				TestEqual(TEXT("OwnedObjects.Num() == 1"), EventData.OwnedObjects.Num(), 1);
+				TestTrue(TEXT("OwnedObjects contains owned object"), EventData.OwnedObjects.Contains(FConcertObjectInStreamID{ SenderStreamId, ObjectReplicator->TestObject }));
+			}
+			else
+			{
+				AddError(TEXT("No activity was produced"));
+			}
+		}
+	END_DEFINE_SPEC(FLeaveReplicationActivitySpec);
+
+	/** This tests that when a client leaves replication, an activity containing their registered streams and authority is produced. */
+	void FLeaveReplicationActivitySpec::Define()
+	{
+		BeforeEach([this]
+		{
+			WorkspaceMock = MakeShared<FReplicationWorkspaceCallInterceptorMock>();
+			ObjectReplicator = MakeUnique<FObjectTestReplicator>();
+			Server = MakeUnique<FReplicationServer>(*this, EConcertSyncSessionFlags::Default_MultiUserSession, WorkspaceMock.ToSharedRef());
+			Client = &Server->ConnectClient();
+			
+			Client->JoinReplication(ObjectReplicator->CreateSenderArgs(SenderStreamId));
+			Client->GetClientReplicationManager().TakeAuthorityOver({ ObjectReplicator->TestObject });
+		});
+		AfterEach([this]
+		{
+			Server.Reset();
+		});
+		
+		It("When leaving replication, an activity is produced", [this]
+		{
+			Client->LeaveReplication();
+			ValidateLeaveActivity();
+		});
+		It("When disconnecting from the session, an activity is produced", [this]
+		{
+			Client->GetClientSessionMock()->Disconnect();
+			ValidateLeaveActivity();
+		});
+	}
+	
+	BEGIN_DEFINE_SPEC(FNoReplicationActivitiesSpec, "Editor.Concert.Replication.RestoreContent.Activity", EAutomationTestFlags::EditorContext | EAutomationTestFlags::EngineFilter)
+		/** Detects calls into the workspace */
+		TSharedPtr<FReplicationWorkspaceCallInterceptorMock> WorkspaceMock;
+	
+		TUniquePtr<FReplicationServer> Server;
+		FReplicationClient* Client = nullptr;
+	END_DEFINE_SPEC(FNoReplicationActivitiesSpec);
+
+	/** This tests no activities are generated when the session EConcertSyncSessionFlags::ShouldEnableReplicationActivities is not set.  */
+	void FNoReplicationActivitiesSpec::Define()
+	{
+		BeforeEach([this]
+		{
+			WorkspaceMock = MakeShared<FReplicationWorkspaceCallInterceptorMock>();
+			Server = MakeUnique<FReplicationServer>(*this, EConcertSyncSessionFlags::Default_MultiUserSession & ~EConcertSyncSessionFlags::ShouldEnableReplicationActivities, WorkspaceMock.ToSharedRef());
+			Client = &Server->ConnectClient();
+
+			Client->JoinReplicationAsListener({});
+		});
+		AfterEach([this]
+		{
+			Server.Reset();
+			WorkspaceMock.Reset();
+		});
+		
+		It("When leaving replication, no activity is produced", [this]
+		{
+			Client->LeaveReplication();
+			TestFalse(TEXT("Activity produced"), WorkspaceMock->LastCall_ProduceClientLeaveReplicationActivity.IsSet());
+		});
+	}
+}
diff --git a/Engine/Plugins/Developer/Concert/ConcertSync/ConcertSyncTest/Source/ConcertSyncTest/Private/Replication/Restore/RestoreContentDisabled.spec.cpp b/Engine/Plugins/Developer/Concert/ConcertSync/ConcertSyncTest/Source/ConcertSyncTest/Private/Replication/Restore/RestoreContentDisabled.spec.cpp
new file mode 100644
--- /dev/null
+++ b/Engine/Plugins/Developer/Concert/ConcertSync/ConcertSyncTest/Source/ConcertSyncTest/Private/Replication/Restore/RestoreContentDisabled.spec.cpp
@@ -0,0 +1,50 @@
+﻿// Copyright Epic Games, Inc. All Rights Reserved.
+
+#include "Replication/Util/Mocks/ReplicationWorkspaceCallInterceptorMock.h"
+#include "Replication/Util/Spec/ReplicationClient.h"
+#include "Replication/Util/Spec/ReplicationServer.h"
+#include "Replication/Util/Spec/ObjectTestReplicator.h"
+
+#include "Misc/AutomationTest.h"
+
+namespace UE::ConcertSyncTests::Replication::RestoreContent
+{
+	BEGIN_DEFINE_SPEC(FRestoreContentDisabledSpec, "Editor.Concert.Replication.RestoreContent", EAutomationTestFlags::EditorContext | EAutomationTestFlags::EngineFilter)
+		/** Detects calls into the workspace */
+		TSharedPtr<FReplicationWorkspaceCallInterceptorMock> WorkspaceMock;
+	
+		TUniquePtr<FReplicationServer> Server;
+		FReplicationClient* Client = nullptr;
+	END_DEFINE_SPEC(FRestoreContentDisabledSpec);
+
+	/** This tests that a client's stream and authority cannot be restored when EConcertSyncSessionFlags::ShouldEnableReplicationActivities is not set. */
+	void FRestoreContentDisabledSpec::Define()
+	{
+		BeforeEach([this]
+		{
+			WorkspaceMock = MakeShared<FReplicationWorkspaceCallInterceptorMock>();
+			Server = MakeUnique<FReplicationServer>(*this, EConcertSyncSessionFlags::Default_MultiUserSession & ~EConcertSyncSessionFlags::ShouldEnableReplicationActivities, WorkspaceMock.ToSharedRef());
+			Client = &Server->ConnectClient();
+
+			Client->JoinReplicationAsListener({});
+		});
+		AfterEach([this]
+		{
+			Server.Reset();
+			WorkspaceMock.Reset();
+		});
+		
+		It("If EConcertSyncSessionFlags::ShouldEnableReplicationActivities, then FConcertReplication_RestoreContent_Request fails.", [this]
+		{
+			bool bReceivedResponse = false;
+			Client->GetClientReplicationManager()
+				.RestoreContent()
+				.Next([this, &bReceivedResponse](FConcertReplication_RestoreContent_Response&& Response)
+				{
+					bReceivedResponse = true;
+					TestEqual(TEXT("Error code"), Response.ErrorCode, EConcertReplicationRestoreErrorCode::NotSupported);
+				});
+			TestTrue(TEXT("Received response"), bReceivedResponse);
+		});
+	}
+}
diff --git a/Engine/Plugins/Developer/Concert/ConcertSync/ConcertSyncTest/Source/ConcertSyncTest/Private/Replication/Restore/RestoreContentPreexisting.spec.cpp b/Engine/Plugins/Developer/Concert/ConcertSync/ConcertSyncTest/Source/ConcertSyncTest/Private/Replication/Restore/RestoreContentPreexisting.spec.cpp
new file mode 100644
--- /dev/null
+++ b/Engine/Plugins/Developer/Concert/ConcertSync/ConcertSyncTest/Source/ConcertSyncTest/Private/Replication/Restore/RestoreContentPreexisting.spec.cpp
@@ -0,0 +1,213 @@
+﻿// Copyright Epic Games, Inc. All Rights Reserved.
+
+#include "Replication/Util/Mocks/ReplicationWorkspaceCallInterceptorMock.h"
+#include "Replication/Util/Spec/ReplicationClient.h"
+#include "Replication/Util/Spec/ReplicationServer.h"
+#include "Replication/Util/Spec/ObjectTestReplicator.h"
+
+#include "Misc/AutomationTest.h"
+#include "Replication/Messages/ReplicationActivity.h"
+#include "Replication/Util/ConcertClientReplicationBridgeMock.h"
+
+namespace UE::ConcertSyncTests::Replication::RestoreContent
+{
+	BEGIN_DEFINE_SPEC(FRestoreContentPreexistingSpec, "Editor.Concert.Replication.RestoreContent", EAutomationTestFlags::EditorContext | EAutomationTestFlags::EngineFilter)
+		/** Detects calls into the workspace */
+		TSharedPtr<FReplicationWorkspaceCallInterceptorMock> WorkspaceMock;
+		FConcertSyncReplicationPayload_LeaveReplication LeaveReplicationData;
+
+		TUniquePtr<FObjectTestReplicator> PreexistingObject1;
+		TUniquePtr<FObjectTestReplicator> RestoredObject1;
+		TUniquePtr<FReplicationServer> Server;
+		FReplicationClient* Client = nullptr;
+		
+		FGuid StreamId = FGuid::NewGuid();
+
+		void RestoreAndTestStreamContent(
+			const ConcertSyncClient::Replication::FJoinReplicatedSessionArgs& JoinArgs,
+			TFunctionRef<void(const TArray<FConcertBaseStreamInfo>& Streams)> TestStreamContent,
+			TFunctionRef<void(const TArray<FConcertAuthorityClientInfo>& Authority)> TestAuthorityContent
+			)
+		{
+			Client->JoinReplication(JoinArgs);
+			
+			InsertActivityData();
+			Client->GetClientReplicationManager().RestoreContent({ .Flags = EConcertReplicationRestoreContentFlags::StreamsAndAuthority | EConcertReplicationRestoreContentFlags::RestoreOnTop });
+			Client->GetClientReplicationManager()
+				.QueryClientInfo({ .ClientEndpointIds = { Client->GetEndpointId() } })
+				.Next([this, &TestStreamContent, &TestAuthorityContent](FConcertReplication_QueryReplicationInfo_Response&& Response)
+				{
+					const FConcertQueriedClientInfo* ClientInfo = Response.ClientInfo.Find(Client->GetEndpointId());
+					if (!ClientInfo)
+					{
+						AddError(TEXT("Not restored correctly"));
+						return;
+					}
+
+					TestStreamContent(ClientInfo->Streams);
+					TestAuthorityContent(ClientInfo->Authority);
+				});
+		}
+	
+		void InsertActivityData() const
+		{
+			WorkspaceMock->ReturnResult_GetLastLeaveReplicationActivityByClient = LeaveReplicationData;
+			WorkspaceMock->ReturnResult_GetLeaveReplicationActivityById = LeaveReplicationData;
+		}
+		FConcertReplicatedObjectInfo& GetRestoredObjectData()
+		{
+			return LeaveReplicationData.Streams[0].BaseDescription.ReplicationMap.ReplicatedObjects[RestoredObject1->TestObject];
+		}
+	END_DEFINE_SPEC(FRestoreContentPreexistingSpec);
+
+	/** This tests that a client's stream and authority can be restored when the client has preexisting stream content and authority. */
+	void FRestoreContentPreexistingSpec::Define()
+	{
+		BeforeEach([this]
+		{
+			WorkspaceMock = MakeShared<FReplicationWorkspaceCallInterceptorMock>();
+			PreexistingObject1 = MakeUnique<FObjectTestReplicator>(TEXT("PreexistingObject1"));
+			RestoredObject1 = MakeUnique<FObjectTestReplicator>(TEXT("RestoredObject1"));
+			Server = MakeUnique<FReplicationServer>(*this, EConcertSyncSessionFlags::Default_MultiUserSession, WorkspaceMock.ToSharedRef());
+			Client = &Server->ConnectClient();
+
+			LeaveReplicationData.Streams.Add(RestoredObject1->CreateStream(StreamId));
+			LeaveReplicationData.OwnedObjects.Add({ StreamId, RestoredObject1->TestObject });
+		});
+		AfterEach([this]
+		{
+			WorkspaceMock.Reset();
+			Server.Reset();
+			PreexistingObject1.Reset();
+			RestoredObject1.Reset();
+			Client = nullptr;
+			LeaveReplicationData = {};
+		});
+
+		It("Aggregate objects into same stream", [this]
+		{
+			const ConcertSyncClient::Replication::FJoinReplicatedSessionArgs JoinArgs = PreexistingObject1->CreateSenderArgs(StreamId);
+			const auto ValidateStream = [this, &JoinArgs](const TArray<FConcertBaseStreamInfo>& Streams)
+			{
+				if (Streams.Num() != 1)
+				{
+					AddError(TEXT("Wrong streams"));
+					return;
+				}
+
+				const FConcertBaseStreamInfo& Stream =Streams[0];
+				const FConcertReplicatedObjectInfo& ExpectedRestoredContent = GetRestoredObjectData();
+				const FConcertReplicatedObjectInfo& ExpectedPreexistingContent = JoinArgs.Streams[0].BaseDescription.ReplicationMap.ReplicatedObjects[PreexistingObject1->TestObject];
+				const FConcertReplicatedObjectInfo* ActualRestoredContent = Stream.ReplicationMap.ReplicatedObjects.Find(RestoredObject1->TestObject);
+				const FConcertReplicatedObjectInfo* ActualPreexistingContent = Stream.ReplicationMap.ReplicatedObjects.Find(PreexistingObject1->TestObject);
+				TestTrue(TEXT("Preexisting"), ActualPreexistingContent && *ActualPreexistingContent == ExpectedPreexistingContent);
+				TestTrue(TEXT("Restored"), ActualRestoredContent && *ActualRestoredContent == ExpectedRestoredContent);
+			};
+			const auto ValidateAuthority = [this](const TArray<FConcertAuthorityClientInfo>& Authority)
+			{
+				if (Authority.Num() != 1)
+				{
+					AddError(TEXT("Wrong authority"));
+					return;
+				}
+
+				TestEqual(TEXT("StreamId"), Authority[0].StreamId, StreamId);
+				TestEqual(TEXT("1 owned object"), Authority[0].AuthoredObjects.Num(), 1);
+				TestTrue(TEXT("Restored object is owned"), Authority[0].AuthoredObjects.Contains(RestoredObject1->TestObject));
+			};
+			RestoreAndTestStreamContent(JoinArgs, ValidateStream, ValidateAuthority);
+		});
+		
+		It("When restoring on top of object that already has all properties, the properties are retained", [this]
+		{
+			const ConcertSyncClient::Replication::FJoinReplicatedSessionArgs JoinArgs = RestoredObject1->CreateSenderArgs(StreamId);
+			FConcertReplicatedObjectInfo& RestoredData = GetRestoredObjectData();
+			RestoredData.PropertySelection.ReplicatedProperties = { *FConcertPropertyChain::CreateFromPath(*RestoredObject1->TestObject->GetClass(), { TEXT("Float") }) };
+			
+			const auto ValidateStream = [this, &JoinArgs](const TArray<FConcertBaseStreamInfo>& Streams)
+			{
+				if (Streams.Num() != 1)
+				{
+					AddError(TEXT("Wrong streams"));
+					return;
+				}
+
+				const FConcertBaseStreamInfo& Stream = Streams[0];
+				const FConcertReplicatedObjectInfo* ActualContent = Stream.ReplicationMap.ReplicatedObjects.Find(RestoredObject1->TestObject);
+				const FConcertReplicatedObjectInfo& ExpectedContent = JoinArgs.Streams[0].BaseDescription.ReplicationMap.ReplicatedObjects[RestoredObject1->TestObject];
+				TestTrue(TEXT("Preexisting"), ActualContent && *ActualContent == ExpectedContent);
+			};
+			const auto ValidateAuthority = [this](const TArray<FConcertAuthorityClientInfo>& Authority)
+			{
+				if (Authority.Num() != 1)
+				{
+					AddError(TEXT("Wrong authority"));
+					return;
+				}
+
+				TestEqual(TEXT("StreamId"), Authority[0].StreamId, StreamId);
+				TestEqual(TEXT("1 owned object"), Authority[0].AuthoredObjects.Num(), 1);
+				TestTrue(TEXT("Restored object is owned"), Authority[0].AuthoredObjects.Contains(RestoredObject1->TestObject));
+			};
+			RestoreAndTestStreamContent(JoinArgs, ValidateStream, ValidateAuthority);
+		});
+
+		It("When restoring on top of an object, the properties are properly aggregated", [this]
+		{
+			const TOptional<FConcertPropertyChain> VectorProperty = FConcertPropertyChain::CreateFromPath(*RestoredObject1->TestObject->GetClass(), { TEXT("Vector") });
+			const TOptional<FConcertPropertyChain> VectorXProperty = FConcertPropertyChain::CreateFromPath(*RestoredObject1->TestObject->GetClass(), { TEXT("Vector"), TEXT("X") });
+			const TOptional<FConcertPropertyChain> FloatProperty = FConcertPropertyChain::CreateFromPath(*RestoredObject1->TestObject->GetClass(), { TEXT("Float") });
+			if (!VectorProperty || !VectorXProperty || !FloatProperty)
+			{
+				AddError(TEXT("Error in FConcertPropertyChain::CreateFromPath"));
+				return;
+			}
+			
+			// RestoredObject will already have "Vector.X" to begin with ...
+			ConcertSyncClient::Replication::FJoinReplicatedSessionArgs JoinArgs = RestoredObject1->CreateSenderArgs(StreamId);
+			FConcertReplicatedObjectInfo& JoinData = JoinArgs.Streams[0].BaseDescription.ReplicationMap.ReplicatedObjects[RestoredObject1->TestObject];
+			JoinData.PropertySelection.ReplicatedProperties = { *VectorProperty, *VectorXProperty };
+			
+			// ... the data to restore is "Float" ...
+			FConcertReplicatedObjectInfo& RestoredContent = GetRestoredObjectData();
+			RestoredContent.PropertySelection.ReplicatedProperties = { *FloatProperty };
+			
+			const auto ValidateStream = [this, &RestoredContent, &VectorProperty, &VectorXProperty, &FloatProperty](const TArray<FConcertBaseStreamInfo>& Streams)
+			{
+				if (Streams.Num() != 1)
+				{
+					AddError(TEXT("Wrong streams"));
+					return;
+				}
+
+				const FConcertBaseStreamInfo& Stream = Streams[0];
+				const FConcertReplicatedObjectInfo* ActualContent = Stream.ReplicationMap.ReplicatedObjects.Find(RestoredObject1->TestObject);
+				if (!ActualContent)
+				{
+					AddError(TEXT("No object data"));
+					return;
+				}
+
+				// ... so the object should now have all 3 properties
+				const TSet<FConcertPropertyChain>& Properties = ActualContent->PropertySelection.ReplicatedProperties;
+				TestEqual(TEXT("Properties.Num() == 3"), Properties.Num(), 3);
+				TestTrue(TEXT("Has property: Vector"), Properties.Contains(*VectorProperty));
+				TestTrue(TEXT("Has property: Vector.X"), Properties.Contains(*VectorXProperty));
+				TestTrue(TEXT("Has property: Float"), Properties.Contains(*FloatProperty));
+			};
+			const auto ValidateAuthority = [this](const TArray<FConcertAuthorityClientInfo>& Authority)
+			{
+				if (Authority.Num() != 1)
+				{
+					AddError(TEXT("Wrong authority"));
+					return;
+				}
+
+				TestEqual(TEXT("StreamId"), Authority[0].StreamId, StreamId);
+				TestEqual(TEXT("1 owned object"), Authority[0].AuthoredObjects.Num(), 1);
+				TestTrue(TEXT("Restored object is owned"), Authority[0].AuthoredObjects.Contains(RestoredObject1->TestObject));
+			};
+			RestoreAndTestStreamContent(JoinArgs, ValidateStream, ValidateAuthority);
+		});
+	}
+}
diff --git a/Engine/Plugins/Developer/Concert/ConcertSync/ConcertSyncTest/Source/ConcertSyncTest/Private/Replication/Restore/RestoreContentWithNoPreexisting.spec.cpp b/Engine/Plugins/Developer/Concert/ConcertSync/ConcertSyncTest/Source/ConcertSyncTest/Private/Replication/Restore/RestoreContentWithNoPreexisting.spec.cpp
new file mode 100644
--- /dev/null
+++ b/Engine/Plugins/Developer/Concert/ConcertSync/ConcertSyncTest/Source/ConcertSyncTest/Private/Replication/Restore/RestoreContentWithNoPreexisting.spec.cpp
@@ -0,0 +1,431 @@
+﻿// Copyright Epic Games, Inc. All Rights Reserved.
+
+#include "Replication/Util/Mocks/ReplicationWorkspaceCallInterceptorMock.h"
+#include "Replication/Util/Spec/ReplicationClient.h"
+#include "Replication/Util/Spec/ReplicationServer.h"
+#include "Replication/Util/Spec/ObjectTestReplicator.h"
+
+#include "Misc/AutomationTest.h"
+#include "Replication/Messages/ReplicationActivity.h"
+#include "Replication/Util/ConcertClientReplicationBridgeMock.h"
+
+namespace UE::ConcertSyncTests::Replication::RestoreContent
+{
+	enum class ERestoreTestFlags
+	{
+		None = 0,
+		TestForAuthority = 1 << 0
+	};
+	ENUM_CLASS_FLAGS(ERestoreTestFlags)
+
+	/** Test that restoring content works correctly when the restoring client has no streams nor authority pre-assigned. */
+	BEGIN_DEFINE_SPEC(FRestoreContentWithNoPreExistingSpec, "Editor.Concert.Replication.RestoreContent", EAutomationTestFlags::EditorContext | EAutomationTestFlags::EngineFilter)
+		/** Detects calls into the workspace */
+		TSharedPtr<FReplicationWorkspaceCallInterceptorMock> WorkspaceMock;
+		FConcertSyncReplicationPayload_LeaveReplication LeaveReplicationData;
+
+		TUniquePtr<FObjectTestReplicator> RestoredObject1;
+		TUniquePtr<FObjectTestReplicator> RestoredObject2;
+		TUniquePtr<FReplicationServer> Server;
+		FReplicationClient* Client = nullptr;
+	
+		FGuid RestoreStreamId = FGuid::NewGuid();
+
+		void InsertActivityData() const
+		{
+			WorkspaceMock->ReturnResult_GetLastLeaveReplicationActivityByClient = LeaveReplicationData;
+			WorkspaceMock->ReturnResult_GetLeaveReplicationActivityById = LeaveReplicationData;
+		}
+
+		TFuture<FConcertReplication_RestoreContent_Response> RestoreThenTestErrorCode(const FConcertReplication_RestoreContent_Request& Request, EConcertReplicationRestoreErrorCode ExpectedErrorCode)
+		{
+			bool bReceivedResponse = false;
+			TFuture<FConcertReplication_RestoreContent_Response> Future = Client->GetClientReplicationManager()
+				.RestoreContent(Request)
+				.Next([this, ExpectedErrorCode, &bReceivedResponse](FConcertReplication_RestoreContent_Response&& Response)
+				{
+					bReceivedResponse = true;
+					TestEqual(TEXT("Error code"), Response.ErrorCode, ExpectedErrorCode);
+					return Response;
+				});
+			TestTrue(TEXT("Received response"), bReceivedResponse);
+			return MoveTemp(Future);
+		}
+
+		void RestoreStreamAndAuthorityThenTest(const FConcertReplication_RestoreContent_Request& Request, ERestoreTestFlags Flags = ERestoreTestFlags::TestForAuthority)
+		{
+			// Join the client first so sync control is granted by the request.
+			FReplicationClient& Receiver = Server->ConnectClient();
+			Receiver.JoinReplicationAsListener({ RestoredObject1->TestObject, RestoredObject2->TestObject });
+
+			// Do the actual restoring
+			InsertActivityData();
+			RestoreThenTestErrorCode(
+				Request,
+				EConcertReplicationRestoreErrorCode::Success
+			)
+			.Next([this, Flags](FConcertReplication_RestoreContent_Response&& Response)
+			{
+				if (EnumHasAnyFlags(Flags, ERestoreTestFlags::TestForAuthority))
+				{
+					TestEqual(TEXT("NewControlStates.Num() == 2"), Response.SyncControl.NewControlStates.Num(), 2);
+					const bool* bNewSyncControl1 = Response.SyncControl.NewControlStates.Find({ RestoreStreamId, RestoredObject1->TestObject });
+					const bool* bNewSyncControl2 = Response.SyncControl.NewControlStates.Find({ RestoreStreamId, RestoredObject2->TestObject });
+					TestTrue(TEXT("Sync Control 1"), bNewSyncControl1 && *bNewSyncControl1);
+					TestTrue(TEXT("Sync Control 2"), bNewSyncControl2 && *bNewSyncControl2);
+				}
+			});
+
+			
+			
+			// Can replicate?
+			if (EnumHasAnyFlags(Flags, ERestoreTestFlags::TestForAuthority))
+			{
+				Client->GetBridgeMock().InjectAvailableObject(*RestoredObject1->TestObject);
+				RestoredObject1->SimulateSendObjectToReceiver(*this, { *Client, *Server, Receiver }, { RestoreStreamId });
+				RestoredObject1->TestValuesWereReplicated(*this);
+			
+				// Need to hide object because RestoredObject2.SimulateSendObjectToReceiver will test that we only receive from RestoredObject2.
+				Client->GetBridgeMock().HideObject(*RestoredObject1->TestObject);
+				Client->GetBridgeMock().InjectAvailableObject(*RestoredObject2->TestObject);
+				RestoredObject2->SimulateSendObjectToReceiver(*this, { *Client, *Server, Receiver }, { RestoreStreamId });
+				RestoredObject2->TestValuesWereReplicated(*this);
+			}
+
+			
+			
+			// Authority and stream state correct?
+			Client->GetClientReplicationManager()
+				.QueryClientInfo({ { Client->GetEndpointId() } })
+				.Next([this, Flags](FConcertReplication_QueryReplicationInfo_Response&& Response)
+				{
+					const FConcertQueriedClientInfo& ClientInfo = Response.ClientInfo[Client->GetEndpointId()];
+					if (ClientInfo.Streams.Num() != 1)
+					{
+						AddError(TEXT("Stream not found"));
+						return;
+					}
+					const FConcertBaseStreamInfo& Stream = ClientInfo.Streams[0];
+					// This catches everything...
+					TestTrue(TEXT("Stream content"), Stream == LeaveReplicationData.Streams[0].BaseDescription);
+					// ... but for easier debugging we'll test some specific things now
+					TestEqual(TEXT("StreamId"), Stream.Identifier, RestoreStreamId);
+					TestEqual(TEXT("Frequency"), Stream.FrequencySettings, LeaveReplicationData.Streams[0].BaseDescription.FrequencySettings);
+
+					if (!EnumHasAnyFlags(Flags, ERestoreTestFlags::TestForAuthority))
+					{
+						return;
+					}
+					
+					if (ClientInfo.Authority.Num() != 1)
+					{
+						AddError(TEXT("Expected client authority"));
+						return;
+					}
+					TestEqual(TEXT("Owns correct stream"), ClientInfo.Authority[0].StreamId, RestoreStreamId);
+					TestEqual(TEXT("Owns 2 objects"), ClientInfo.Authority[0].AuthoredObjects.Num(), 2);
+					TestTrue(TEXT("Owns correct RestoredObject1"), ClientInfo.HasAuthority({ RestoreStreamId, RestoredObject1->TestObject }));
+					TestTrue(TEXT("Owns correct RestoredObject2"), ClientInfo.HasAuthority({ RestoreStreamId, RestoredObject2->TestObject }));
+				});
+		}
+
+		static FConcertClientInfo MakeClientInfo() { return FConcertClientInfo { .DeviceName = TEXT("MainDeviceName"), .DisplayName = TEXT("MainClientName") }; }
+	END_DEFINE_SPEC(FRestoreContentWithNoPreExistingSpec);
+
+	/** This tests the base functionality of restoring stream and authority content: when the client has an empty stream at the time of restore. */
+	void FRestoreContentWithNoPreExistingSpec::Define()
+	{
+		BeforeEach([this]
+		{
+			WorkspaceMock = MakeShared<FReplicationWorkspaceCallInterceptorMock>();
+			RestoredObject1 = MakeUnique<FObjectTestReplicator>(TEXT("Object1"));
+			RestoredObject2 = MakeUnique<FObjectTestReplicator>(TEXT("Object2"));
+			Server = MakeUnique<FReplicationServer>(*this, EConcertSyncSessionFlags::Default_MultiUserSession, WorkspaceMock.ToSharedRef());
+			Client = &Server->ConnectClient(MakeClientInfo());
+
+			FConcertReplicationStream Stream = RestoredObject1->CreateStream(RestoreStreamId, EConcertObjectReplicationMode::SpecifiedRate, 21);
+			RestoredObject2->AddToStream(Stream, EConcertObjectReplicationMode::SpecifiedRate, 42);
+			LeaveReplicationData.Streams.Add(Stream);
+			LeaveReplicationData.OwnedObjects.Add({ RestoreStreamId, RestoredObject1->TestObject });
+			LeaveReplicationData.OwnedObjects.Add({ RestoreStreamId, RestoredObject2->TestObject });
+			
+			Client->JoinReplication();
+		});
+		AfterEach([this]
+		{
+			WorkspaceMock.Reset();
+			Server.Reset();
+			RestoredObject1.Reset();
+			RestoredObject2.Reset();
+			Client = nullptr;
+			LeaveReplicationData = {};
+		});
+
+		Describe("Restores content correctly", [this]
+		{
+			// Try all permutations of flags
+			const auto BuildTests = [this](EConcertReplicationAuthorityRestoreMode Mode)
+			{
+				It("StreamsAndAuthority", [this, Mode]
+				{
+					const EConcertReplicationRestoreContentFlags Flags = EConcertReplicationRestoreContentFlags::StreamsAndAuthority;
+					RestoreStreamAndAuthorityThenTest({ .Flags = Flags, .AuthorityRestorationMode = Mode });
+				});
+				It("StreamsAndAuthority | RestoreOnTop", [this, Mode]
+				{
+					const EConcertReplicationRestoreContentFlags Flags = EConcertReplicationRestoreContentFlags::StreamsAndAuthority | EConcertReplicationRestoreContentFlags::RestoreOnTop;
+					RestoreStreamAndAuthorityThenTest({ .Flags = Flags, .AuthorityRestorationMode = Mode });
+				});
+				It("StreamsAndAuthority | ValidateUniqueClient", [this, Mode]
+				{
+					const EConcertReplicationRestoreContentFlags Flags = EConcertReplicationRestoreContentFlags::StreamsAndAuthority | EConcertReplicationRestoreContentFlags::ValidateUniqueClient;
+					RestoreStreamAndAuthorityThenTest({ .Flags = Flags, .AuthorityRestorationMode = Mode });
+				});
+				It("StreamsAndAuthority | RestoreOnTop | ValidateUniqueClient", [this, Mode]
+				{
+					const EConcertReplicationRestoreContentFlags Flags = EConcertReplicationRestoreContentFlags::StreamsAndAuthority | EConcertReplicationRestoreContentFlags::ValidateUniqueClient;
+					RestoreStreamAndAuthorityThenTest({ .Flags = Flags, .AuthorityRestorationMode = Mode });
+				});
+
+				Describe("Streams only", [this, Mode]
+				{
+					AfterEach([this]
+					{
+						TestTrue(TEXT("Client thinks it has no authority"), Client->GetClientReplicationManager().GetClientOwnedObjects().IsEmpty());
+						TestTrue(TEXT("Client thinks it has no sync control"), Client->GetClientReplicationManager().GetSyncControlledObjects().IsEmpty());
+						
+						Client->GetClientReplicationManager()
+							.QueryClientInfo({ .ClientEndpointIds = { Client->GetEndpointId() } })
+							.Next([this](FConcertReplication_QueryReplicationInfo_Response&& Response)
+							{
+								if (const FConcertQueriedClientInfo* ClientInfo = Response.ClientInfo.Find(Client->GetEndpointId()))
+								{
+									TestTrue(TEXT("No authority on server"), ClientInfo->Authority.IsEmpty());
+								}
+								else
+								{
+									AddError(TEXT("Streams were not restored"));
+								}
+							});
+					});
+					
+					It("None", [this, Mode]
+					{
+						const EConcertReplicationRestoreContentFlags Flags = EConcertReplicationRestoreContentFlags::None;
+						RestoreStreamAndAuthorityThenTest({ .Flags = Flags, .AuthorityRestorationMode = Mode }, ERestoreTestFlags::None);
+					});
+					It("RestoreOnTop", [this, Mode]
+					{
+						const EConcertReplicationRestoreContentFlags Flags = EConcertReplicationRestoreContentFlags::RestoreOnTop;
+						RestoreStreamAndAuthorityThenTest({ .Flags = Flags, .AuthorityRestorationMode = Mode }, ERestoreTestFlags::None);
+					});
+					It("ValidateUniqueClient", [this, Mode]
+					{
+						const EConcertReplicationRestoreContentFlags Flags = EConcertReplicationRestoreContentFlags::ValidateUniqueClient;
+						RestoreStreamAndAuthorityThenTest({ .Flags = Flags, .AuthorityRestorationMode = Mode }, ERestoreTestFlags::None);
+					});
+					It("RestoreOnTop | ValidateUniqueClient", [this, Mode]
+					{
+						const EConcertReplicationRestoreContentFlags Flags = EConcertReplicationRestoreContentFlags::RestoreOnTop | EConcertReplicationRestoreContentFlags::ValidateUniqueClient;
+						RestoreStreamAndAuthorityThenTest({ .Flags = Flags, .AuthorityRestorationMode = Mode }, ERestoreTestFlags::None);
+					});
+				});
+			};
+			
+			Describe("ExcludeAlreadyOwnedObjectPropertiesFromStream", [this, BuildTests]
+			{
+				BuildTests(EConcertReplicationAuthorityRestoreMode::ExcludeAlreadyOwnedObjectPropertiesFromStream);
+			});
+			Describe("IncludeAlreadyOwnedObjectPropertiesInStream", [this, BuildTests]
+			{
+				BuildTests(EConcertReplicationAuthorityRestoreMode::IncludeAlreadyOwnedObjectPropertiesInStream);
+			});
+			Describe("AllOrNothing", [this, BuildTests]
+			{
+				BuildTests(EConcertReplicationAuthorityRestoreMode::AllOrNothing);
+			});
+		});
+
+		Describe("ClientInfo", [this]
+		{
+			It("Is contained when EConcertReplicationRestoreContentFlagsSendNewState is set", [this]
+			{
+				InsertActivityData();
+				RestoreThenTestErrorCode(
+					{ .Flags = EConcertReplicationRestoreContentFlags::StreamsAndAuthority | EConcertReplicationRestoreContentFlags::SendNewState },
+					EConcertReplicationRestoreErrorCode::Success
+					)
+				.Next([this](FConcertReplication_RestoreContent_Response&& Response)
+				{
+					const FConcertQueriedClientInfo& ClientInfo = Response.ClientInfo;
+					if (ClientInfo.Streams.Num() != 1)
+					{
+						AddError(TEXT("Wrong streams"));
+						return;
+					}
+					TestEqual(TEXT("Stream Content"), ClientInfo.Streams[0], LeaveReplicationData.Streams[0].BaseDescription);
+
+					if (ClientInfo.Authority.Num() != 1)
+					{
+						AddError(TEXT("Wrong authority"));
+						return;
+					}
+					TestEqual(TEXT("Authority StreamId"), ClientInfo.Authority[0].StreamId, RestoreStreamId);
+					TestTrue(TEXT("Authority RestoredObject1"), ClientInfo.Authority[0].AuthoredObjects.Contains(RestoredObject1->TestObject));
+					TestTrue(TEXT("Authority RestoredObject2"), ClientInfo.Authority[0].AuthoredObjects.Contains(RestoredObject2->TestObject));
+				});
+			});
+
+			It("Is not contained when EConcertReplicationRestoreContentFlagsSendNewState is not set", [this]
+			{
+				InsertActivityData();
+				RestoreThenTestErrorCode(
+					{ .Flags = EConcertReplicationRestoreContentFlags::StreamsAndAuthority },
+					EConcertReplicationRestoreErrorCode::Success
+					)
+				.Next([this](FConcertReplication_RestoreContent_Response&& Response)
+				{
+					TestTrue(TEXT("No ClientInfo"), Response.ClientInfo.IsEmpty());
+				});
+			});
+		});
+
+		Describe("When another client has authority", [this]
+		{
+			BeforeEach([this]
+			{
+				InsertActivityData();
+				
+				FReplicationClient& OtherClient = Server->ConnectClient();
+				OtherClient.JoinReplication(RestoredObject1->CreateSenderArgs());
+				OtherClient.GetClientReplicationManager().TakeAuthorityOver({ RestoredObject1->TestObject });
+			});
+			
+			It("When request has EConcertReplicationAuthorityRestoreMode::ExcludeAlreadyOwnedObjectsFromStream, the object is excluded from the stream", [this]
+			{
+				RestoreThenTestErrorCode(
+					{ .Flags = EConcertReplicationRestoreContentFlags::StreamsAndAuthority, .AuthorityRestorationMode = EConcertReplicationAuthorityRestoreMode::ExcludeAlreadyOwnedObjectPropertiesFromStream },
+					EConcertReplicationRestoreErrorCode::Success
+				);
+
+				Client->GetClientReplicationManager()
+					.QueryClientInfo({ { Client->GetEndpointId() } })
+					.Next([this](FConcertReplication_QueryReplicationInfo_Response&& Response)
+					{
+						const FConcertQueriedClientInfo& ClientInfo = Response.ClientInfo[Client->GetEndpointId()];
+						if (ClientInfo.Authority.Num() != 1)
+						{
+							AddError(TEXT("Expected client authority"));
+							return;
+						}
+						TestEqual(TEXT("Owns correct stream"), ClientInfo.Authority[0].StreamId, RestoreStreamId);
+						TestEqual(TEXT("Owns 1 objects"), ClientInfo.Authority[0].AuthoredObjects.Num(), 1);
+						TestTrue(TEXT("Owns correct RestoredObject2"), ClientInfo.HasAuthority({ RestoreStreamId, RestoredObject2->TestObject }));
+					
+						TestEqual(TEXT("No streams"), ClientInfo.Streams.Num(), 1);
+						FConcertBaseStreamInfo ExpectedStream = LeaveReplicationData.Streams[0].BaseDescription;
+						ExpectedStream.ReplicationMap.ReplicatedObjects.Remove(RestoredObject1->TestObject);
+						ExpectedStream.FrequencySettings.ObjectOverrides.Remove(RestoredObject1->TestObject);
+						const FConcertBaseStreamInfo* RestoreStreamInfo = ClientInfo.Streams.FindByPredicate([this](const FConcertBaseStreamInfo& StreamInfo){ return StreamInfo.Identifier == RestoreStreamId; });
+						TestTrue(TEXT("Stream content"), RestoreStreamInfo && *RestoreStreamInfo == ExpectedStream);
+					});
+			});
+			It("When request has EConcertReplicationAuthorityRestoreMode::IncludeAlreadyOwnedObjectsInStream, the object is included in the stream", [this]
+			{
+				RestoreThenTestErrorCode(
+					{ .Flags = EConcertReplicationRestoreContentFlags::StreamsAndAuthority, .AuthorityRestorationMode = EConcertReplicationAuthorityRestoreMode::IncludeAlreadyOwnedObjectPropertiesInStream },
+					EConcertReplicationRestoreErrorCode::Success
+				);
+				
+				Client->GetClientReplicationManager()
+					.QueryClientInfo({ { Client->GetEndpointId() } })
+					.Next([this](FConcertReplication_QueryReplicationInfo_Response&& Response)
+					{
+						const FConcertQueriedClientInfo& ClientInfo = Response.ClientInfo[Client->GetEndpointId()];
+						if (ClientInfo.Authority.Num() != 1)
+						{
+							AddError(TEXT("Expected client authority"));
+							return;
+						}
+						TestEqual(TEXT("Owns correct stream"), ClientInfo.Authority[0].StreamId, RestoreStreamId);
+						TestEqual(TEXT("Owns 1 object"), ClientInfo.Authority[0].AuthoredObjects.Num(), 1);
+						TestTrue(TEXT("Owns correct RestoredObject2"), ClientInfo.HasAuthority({ RestoreStreamId, RestoredObject2->TestObject }));
+						
+						TestEqual(TEXT("1 stream"), ClientInfo.Streams.Num(), 1);
+						const FConcertBaseStreamInfo* RestoreStreamInfo = ClientInfo.Streams.FindByPredicate([this](const FConcertBaseStreamInfo& StreamInfo){ return StreamInfo.Identifier == RestoreStreamId; });
+						TestTrue(TEXT("Stream content"), RestoreStreamInfo && *RestoreStreamInfo == LeaveReplicationData.Streams[0].BaseDescription);
+					});
+			});
+			It("When request has EConcertReplicationAuthorityRestoreMode::AllOrNothing, the request fails", [this]
+			{
+				RestoreThenTestErrorCode({ .AuthorityRestorationMode = EConcertReplicationAuthorityRestoreMode::AllOrNothing }, EConcertReplicationRestoreErrorCode::AuthorityConflict);
+				RestoreThenTestErrorCode(
+					{ .Flags = EConcertReplicationRestoreContentFlags::None, .AuthorityRestorationMode = EConcertReplicationAuthorityRestoreMode::AllOrNothing },
+					EConcertReplicationRestoreErrorCode::AuthorityConflict
+				);
+			});
+		});
+		
+		It("When there is the 'same' client and request has ValidateUniqueClient flag, the error code is NameConflict", [this]
+		{
+			InsertActivityData();
+			Server->ConnectClient(MakeClientInfo());
+			RestoreThenTestErrorCode(
+				{ .Flags = EConcertReplicationRestoreContentFlags::ValidateUniqueClient },
+				EConcertReplicationRestoreErrorCode::NameConflict
+				);
+		});
+		It("When ActivityId references unknown activity, the error code is InvalidActivity.", [this]
+		{
+			constexpr int64 InvalidActivityId = 42;
+			RestoreThenTestErrorCode(
+				{ .ActivityId = { InvalidActivityId } },
+				EConcertReplicationRestoreErrorCode::NoSuchActivity
+				);
+		});
+		It("When ActivityId is left unset and there is no activity, the error code is Success.", [this]
+		{
+			RestoreThenTestErrorCode(
+				{ },
+				EConcertReplicationRestoreErrorCode::Success
+				);
+		});
+	}
+	
+	BEGIN_DEFINE_SPEC(FRestoreContentPropertiesSpec, "Editor.Concert.Replication.RestoreContent", EAutomationTestFlags::EditorContext | EAutomationTestFlags::EngineFilter)
+		/** Detects calls into the workspace */
+		TSharedPtr<FReplicationWorkspaceCallInterceptorMock> WorkspaceMock;
+		FConcertSyncReplicationPayload_LeaveReplication LeaveReplicationData;
+
+		/** The stream to restore */
+		TUniquePtr<FObjectTestReplicator> RestoreObjectReplicator;
+		/** The stream that the client joins with */
+		TUniquePtr<FObjectTestReplicator> JoinObjectReplicator;
+		TUniquePtr<FReplicationServer> Server;
+		FReplicationClient* Client = nullptr;
+	
+		FGuid RestoreStreamId = FGuid::NewGuid();
+		FGuid JoinStreamId = FGuid::NewGuid();
+
+		static FConcertClientInfo MakeClientInfo() { return FConcertClientInfo { .DeviceName = TEXT("MainDeviceName"), .DisplayName = TEXT("MainClientName") }; }
+		static FName OriginalObjectName() { return TEXT("OriginalObjectName"); }
+		static FName RestoreObjectName() { return TEXT("OriginalObjectName"); }
+	END_DEFINE_SPEC(FRestoreContentPropertiesSpec);
+
+	/** This tests that a client's stream and authority can be restored when EConcertSyncSessionFlags::ShouldEnableReplicationActivities is set. */
+	void FRestoreContentPropertiesSpec::Define()
+	{
+		// TODO: Replaces correctly
+		// TODO: Test RestoreOnTop combines correctly (in same stream)
+		// TODO: Test client 2 has authority over non-overlapping properties
+		/*
+		It("Restores correctly (EConcertReplicationRestoreContentFlags::RestoreOnTop)", [this]
+		{
+			TestErrorCode(
+				{ .Flags = EConcertReplicationRestoreContentFlags::Default | EConcertReplicationRestoreContentFlags::RestoreOnTop },
+				EConcertReplicationRestoreErrorCode::Success
+			);
+			TestHasRestored(ERestoreTestFlags::ShouldHaveAuthority);
+		});*/
+	}
+}
diff --git a/Engine/Plugins/Developer/Concert/ConcertSync/ConcertSyncTest/Source/ConcertSyncTest/Private/Replication/SyncControl.spec.cpp b/Engine/Plugins/Developer/Concert/ConcertSync/ConcertSyncTest/Source/ConcertSyncTest/Private/Replication/SyncControl.spec.cpp
--- a/Engine/Plugins/Developer/Concert/ConcertSync/ConcertSyncTest/Source/ConcertSyncTest/Private/Replication/SyncControl.spec.cpp
+++ b/Engine/Plugins/Developer/Concert/ConcertSync/ConcertSyncTest/Source/ConcertSyncTest/Private/Replication/SyncControl.spec.cpp
@@ -289,7 +289,7 @@ namespace UE::ConcertSyncTests::Replication::UI
 					.Next([this, &Manager, &ReleaseAuthority, &SyncControlBeforeRelease](FConcertReplication_ChangeAuthority_Response&& Response)
 					{
 						// Aggregate is validated in SyncControlState.spec.cpp
-						SyncControlBeforeRelease.AppendChanges(ReleaseAuthority, Response);
+						SyncControlBeforeRelease.AppendAuthorityChange(ReleaseAuthority, Response);
 						TestTrue(TEXT("Client predicted correctly"), SyncControlBeforeRelease == Manager.GetSyncControlledObjects());
 					});
 
@@ -300,7 +300,7 @@ namespace UE::ConcertSyncTests::Replication::UI
 					.Next([this, &Manager, &SyncControlBeforeTake, &TakeAuthority](FConcertReplication_ChangeAuthority_Response&& Response)
 					{
 						// Aggregate is validated in SyncControlState.spec.cpp
-						SyncControlBeforeTake.AppendChanges(TakeAuthority, Response);
+						SyncControlBeforeTake.AppendAuthorityChange(TakeAuthority, Response);
 						TestTrue(TEXT("Client predicted correctly"), SyncControlBeforeTake == Manager.GetSyncControlledObjects());
 					});
 			});
@@ -317,7 +317,7 @@ namespace UE::ConcertSyncTests::Replication::UI
 					.Next([this, &Manager, &StreamChange, &SyncControlBeforeRelease](FConcertReplication_ChangeStream_Response&& Response)
 					{
 						// Aggregate is validated in SyncControlState.spec.cpp
-						SyncControlBeforeRelease.AppendChanges(StreamChange);
+						SyncControlBeforeRelease.AppendStreamChange(StreamChange);
 						TestTrue(TEXT("Client predicted correctly"), SyncControlBeforeRelease == Manager.GetSyncControlledObjects());
 					});
 			});
diff --git a/Engine/Plugins/Developer/Concert/ConcertSync/ConcertSyncTest/Source/ConcertSyncTest/Private/Replication/Util/ConcertClientReplicationBridgeMock.h b/Engine/Plugins/Developer/Concert/ConcertSync/ConcertSyncTest/Source/ConcertSyncTest/Private/Replication/Util/ConcertClientReplicationBridgeMock.h
--- a/Engine/Plugins/Developer/Concert/ConcertSync/ConcertSyncTest/Source/ConcertSyncTest/Private/Replication/Util/ConcertClientReplicationBridgeMock.h
+++ b/Engine/Plugins/Developer/Concert/ConcertSync/ConcertSyncTest/Source/ConcertSyncTest/Private/Replication/Util/ConcertClientReplicationBridgeMock.h
@@ -23,6 +23,11 @@ namespace UE::ConcertSyncTests::Replication
 			AvailableObjects.Add(&Object);
 			OnObjectDiscoveredDelegate.Broadcast(Object);
 		}
+		void HideObject(UObject& Object)
+		{
+			AvailableObjects.Remove(&Object);
+			OnObjectRemovedDelegate.Broadcast(&Object);
+		}
 		
 		virtual void PushTrackedObjects(TArrayView<const FSoftObjectPath> InTrackedObjects) override
 		{
diff --git a/Engine/Plugins/Developer/Concert/ConcertSync/ConcertSyncTest/Source/ConcertSyncTest/Private/Replication/Util/Mocks/ReplicationWorkspaceCallDetectorMock.h b/Engine/Plugins/Developer/Concert/ConcertSync/ConcertSyncTest/Source/ConcertSyncTest/Private/Replication/Util/Mocks/ReplicationWorkspaceCallDetectorMock.h
deleted file mode 100644
diff --git a/Engine/Plugins/Developer/Concert/ConcertSync/ConcertSyncTest/Source/ConcertSyncTest/Private/Replication/Util/Mocks/ReplicationWorkspaceCallInterceptorMock.h b/Engine/Plugins/Developer/Concert/ConcertSync/ConcertSyncTest/Source/ConcertSyncTest/Private/Replication/Util/Mocks/ReplicationWorkspaceCallInterceptorMock.h
new file mode 100644
--- /dev/null
+++ b/Engine/Plugins/Developer/Concert/ConcertSync/ConcertSyncTest/Source/ConcertSyncTest/Private/Replication/Util/Mocks/ReplicationWorkspaceCallInterceptorMock.h
@@ -0,0 +1,56 @@
+﻿// Copyright Epic Games, Inc. All Rights Reserved.
+
+#pragma once
+
+#include "Replication/IReplicationWorkspace.h"
+#include "Replication/Messages/ReplicationActivity.h"
+#include "Replication/Messages/RestoreContent.h"
+
+#include "Misc/Optional.h"
+#include "Misc/Guid.h"
+#include "Templates/Tuple.h"
+
+namespace UE::ConcertSyncTests::Replication
+{
+	class FReplicationWorkspaceCallInterceptorMock : public ConcertSyncServer::Replication::IReplicationWorkspace
+	{
+	public:
+
+		/** Arguments of last ProduceClientLeaveReplicationActivity call. */
+		TOptional<TTuple<FGuid, FConcertSyncReplicationPayload_LeaveReplication>> LastCall_ProduceClientLeaveReplicationActivity;
+		/** Arguments of last GetLastLeaveReplicationActivityByClient call. */
+		mutable TOptional<TTuple<FConcertClientInfo>> LastCall_GetLastLeaveReplicationActivityByClient;
+		/** Arguments of last GetLastLeaveReplicationActivityByClient call. */
+		mutable TOptional<TTuple<int64>> LastCall_GetLeaveReplicationActivityById;
+
+		/** The result to return in GetLastLeaveReplicationActivityByClient. */
+		TOptional<FConcertSyncReplicationPayload_LeaveReplication> ReturnResult_GetLastLeaveReplicationActivityByClient;
+		/** The result to return in GetLastLeaveReplicationActivityByClient. */
+		TOptional<FConcertSyncReplicationPayload_LeaveReplication> ReturnResult_GetLeaveReplicationActivityById;
+		
+		//~ Begin IReplicationWorkspace Interface
+		virtual void ProduceClientLeaveReplicationActivity(const FGuid& EndpointId, const FConcertSyncReplicationPayload_LeaveReplication& EventData) override
+		{
+			LastCall_ProduceClientLeaveReplicationActivity = MakeTuple(EndpointId, EventData);
+		}
+		virtual bool GetLastLeaveReplicationActivityByClient(const FConcertClientInfo& InClientInfo, FConcertSyncReplicationPayload_LeaveReplication& OutLeaveReplication) const override
+		{
+			LastCall_GetLastLeaveReplicationActivityByClient = MakeTuple(InClientInfo);
+			if (ReturnResult_GetLastLeaveReplicationActivityByClient)
+			{
+				OutLeaveReplication = *ReturnResult_GetLastLeaveReplicationActivityByClient;
+			}
+			return ReturnResult_GetLastLeaveReplicationActivityByClient.IsSet();
+		}
+		virtual bool GetLeaveReplicationActivityById(const int64 ActivityId, FConcertSyncReplicationPayload_LeaveReplication& OutLeaveReplication) const override
+		{
+			LastCall_GetLeaveReplicationActivityById = MakeTuple(ActivityId);
+			if (ReturnResult_GetLeaveReplicationActivityById)
+			{
+				OutLeaveReplication = *ReturnResult_GetLeaveReplicationActivityById;
+			}
+			return ReturnResult_GetLeaveReplicationActivityById.IsSet();
+		}
+		//~ End IReplicationWorkspace Interface
+	};
+}
diff --git a/Engine/Plugins/Developer/Concert/ConcertSync/ConcertSyncTest/Source/ConcertSyncTest/Private/Replication/Util/Mocks/ReplicationWorkspaceEmptyMock.h b/Engine/Plugins/Developer/Concert/ConcertSync/ConcertSyncTest/Source/ConcertSyncTest/Private/Replication/Util/Mocks/ReplicationWorkspaceEmptyMock.h
--- a/Engine/Plugins/Developer/Concert/ConcertSync/ConcertSyncTest/Source/ConcertSyncTest/Private/Replication/Util/Mocks/ReplicationWorkspaceEmptyMock.h
+++ b/Engine/Plugins/Developer/Concert/ConcertSync/ConcertSyncTest/Source/ConcertSyncTest/Private/Replication/Util/Mocks/ReplicationWorkspaceEmptyMock.h
@@ -1,6 +1,7 @@
 ﻿// Copyright Epic Games, Inc. All Rights Reserved.
 
 #pragma once
+
 #include "Replication/IReplicationWorkspace.h"
 
 namespace UE::ConcertSyncTests::Replication
@@ -11,6 +12,8 @@ namespace UE::ConcertSyncTests::Replication
 		
 		//~ Begin IReplicationWorkspace Interface
 		virtual void ProduceClientLeaveReplicationActivity(const FGuid&, const FConcertSyncReplicationPayload_LeaveReplication&) override {}
+		virtual bool GetLastLeaveReplicationActivityByClient(const FConcertClientInfo& InClientInfo, FConcertSyncReplicationPayload_LeaveReplication& OutLeaveReplication) const override { return false; }
+		virtual bool GetLeaveReplicationActivityById(const int64 ActivityId, FConcertSyncReplicationPayload_LeaveReplication& OutLeaveReplication) const override { return false; }
 		//~ End IReplicationWorkspace Interface
 	};
 }
diff --git a/Engine/Plugins/Developer/Concert/ConcertSync/ConcertSyncTest/Source/ConcertSyncTest/Private/Replication/Util/Spec/ObjectTestReplicator.cpp b/Engine/Plugins/Developer/Concert/ConcertSync/ConcertSyncTest/Source/ConcertSyncTest/Private/Replication/Util/Spec/ObjectTestReplicator.cpp
--- a/Engine/Plugins/Developer/Concert/ConcertSync/ConcertSyncTest/Source/ConcertSyncTest/Private/Replication/Util/Spec/ObjectTestReplicator.cpp
+++ b/Engine/Plugins/Developer/Concert/ConcertSync/ConcertSyncTest/Source/ConcertSyncTest/Private/Replication/Util/Spec/ObjectTestReplicator.cpp
@@ -43,6 +43,15 @@ namespace UE::ConcertSyncTests::Replication
 	}
 
 	FConcertReplicationStream FObjectTestReplicator::CreateStream(FGuid SenderStreamId, EConcertObjectReplicationMode ReplicationMode, uint8 ReplicationRate) const
+	{
+		FConcertReplicationStream SendingStream;
+		SendingStream.BaseDescription.Identifier = SenderStreamId;
+		SendingStream.BaseDescription.FrequencySettings.Defaults = { ReplicationMode, ReplicationRate };
+		AddToStream(SendingStream);
+		return SendingStream;
+	}
+
+	void FObjectTestReplicator::AddToStream(FConcertReplicationStream& Stream, EConcertObjectReplicationMode ReplicationMode, uint8 ReplicationRate) const
 	{
 		FConcertReplicatedObjectInfo ReplicatedObjectInfo { TestObject->GetClass() };
 		ConcertSyncCore::PropertyChain::ForEachReplicatableConcertProperty(*TestObject->GetClass(), [&ReplicatedObjectInfo](FConcertPropertyChain&& Chain)
@@ -50,12 +59,13 @@ namespace UE::ConcertSyncTests::Replication
 			ReplicatedObjectInfo.PropertySelection.ReplicatedProperties.Emplace(MoveTemp(Chain));
 			return EBreakBehavior::Continue;
 		});
+		Stream.BaseDescription.ReplicationMap.ReplicatedObjects.Add(TestObject, ReplicatedObjectInfo);
 
-		FConcertReplicationStream SendingStream;
-		SendingStream.BaseDescription.Identifier = SenderStreamId;
-		SendingStream.BaseDescription.ReplicationMap.ReplicatedObjects.Add(TestObject, ReplicatedObjectInfo);
-		SendingStream.BaseDescription.FrequencySettings.Defaults = { ReplicationMode, ReplicationRate };
-		return SendingStream;
+		const FConcertObjectReplicationSettings FrequencySettings { ReplicationMode, ReplicationRate };
+		if (Stream.BaseDescription.FrequencySettings.Defaults != FrequencySettings)
+		{
+			Stream.BaseDescription.FrequencySettings.ObjectOverrides.Add(TestObject, FrequencySettings);
+		}
 	}
 
 	void FObjectTestReplicator::SimulateSendObjectToReceiver(
diff --git a/Engine/Plugins/Developer/Concert/ConcertSync/ConcertSyncTest/Source/ConcertSyncTest/Private/Replication/Util/Spec/ObjectTestReplicator.h b/Engine/Plugins/Developer/Concert/ConcertSync/ConcertSyncTest/Source/ConcertSyncTest/Private/Replication/Util/Spec/ObjectTestReplicator.h
--- a/Engine/Plugins/Developer/Concert/ConcertSync/ConcertSyncTest/Source/ConcertSyncTest/Private/Replication/Util/Spec/ObjectTestReplicator.h
+++ b/Engine/Plugins/Developer/Concert/ConcertSync/ConcertSyncTest/Source/ConcertSyncTest/Private/Replication/Util/Spec/ObjectTestReplicator.h
@@ -80,6 +80,12 @@ namespace UE::ConcertSyncTests::Replication
 			EConcertObjectReplicationMode ReplicationMode = EConcertObjectReplicationMode::Realtime,
 			uint8 ReplicationRate = 30
 		) const;
+		/** Util adding replication settings for TestObject to a stream */
+		void AddToStream(
+			FConcertReplicationStream& Stream,
+			EConcertObjectReplicationMode ReplicationMode = EConcertObjectReplicationMode::Realtime,
+			uint8 ReplicationRate = 30
+			) const;
 		
 		/** This overload allows you to send properties but from multiple streams. */
 		void SimulateSendObjectToReceiver(
diff --git a/Engine/Plugins/Developer/Concert/ConcertSync/ConcertSyncTest/Source/ConcertSyncTest/Private/Replication/Util/Spec/ReplicationClient.h b/Engine/Plugins/Developer/Concert/ConcertSync/ConcertSyncTest/Source/ConcertSyncTest/Private/Replication/Util/Spec/ReplicationClient.h
--- a/Engine/Plugins/Developer/Concert/ConcertSync/ConcertSyncTest/Source/ConcertSyncTest/Private/Replication/Util/Spec/ReplicationClient.h
+++ b/Engine/Plugins/Developer/Concert/ConcertSync/ConcertSyncTest/Source/ConcertSyncTest/Private/Replication/Util/Spec/ReplicationClient.h
@@ -2,6 +2,7 @@
 
 #pragma once
 
+#include "ConcertMessageData.h"
 #include "ConcertSyncSessionFlags.h"
 #include "Replication/IConcertClientReplicationManager.h"
 #include "Util/ClientServerCommunicationTest.h"
@@ -34,11 +35,11 @@ namespace UE::ConcertSyncTests::Replication
 	{
 	public:
 		
-		FReplicationClient(const FGuid& ClientEndPointId, EConcertSyncSessionFlags SessionFlags, FConcertServerSessionMock& Server, FAutomationTestBase& TestContext)
+		FReplicationClient(const FGuid& ClientEndPointId, EConcertSyncSessionFlags SessionFlags, FConcertServerSessionMock& Server, FAutomationTestBase& TestContext, FConcertClientInfo ClientInfo = {})
 			: ClientEndPointId(ClientEndPointId)
 			, SessionFlags(SessionFlags)
 			, TestContext(TestContext)
-			, ClientSessionMock(MakeShared<FConcertClientSessionMock>(ClientEndPointId, Server))
+			, ClientSessionMock(MakeShared<FConcertClientSessionMock>(ClientEndPointId, Server, MoveTemp(ClientInfo)))
 		{}
 		
 		/** Lets the client process any messages that have come in. */
@@ -58,6 +59,8 @@ namespace UE::ConcertSyncTests::Replication
 		void LeaveReplication() const;
 
 		const FGuid& GetEndpointId() const { return ClientEndPointId; }
+		const FConcertClientInfo& GetClientInfo() const { return ClientSessionMock->GetLocalClientInfo(); }
+		
 		TSharedRef<FConcertClientSessionBaseMock> GetClientSessionMock() const { return ClientSessionMock; }
 		FConcertClientReplicationBridgeMock& GetBridgeMock() const { checkf(BridgeMock, TEXT("You forgot to call JoinReplication")); return *BridgeMock; }
 		IConcertClientReplicationManager& GetClientReplicationManager() const { checkf(ClientReplicationManager, TEXT("You forgot to call JoinReplication")); return *ClientReplicationManager; }
diff --git a/Engine/Plugins/Developer/Concert/ConcertSync/ConcertSyncTest/Source/ConcertSyncTest/Private/Replication/Util/Spec/ReplicationServer.cpp b/Engine/Plugins/Developer/Concert/ConcertSync/ConcertSyncTest/Source/ConcertSyncTest/Private/Replication/Util/Spec/ReplicationServer.cpp
--- a/Engine/Plugins/Developer/Concert/ConcertSync/ConcertSyncTest/Source/ConcertSyncTest/Private/Replication/Util/Spec/ReplicationServer.cpp
+++ b/Engine/Plugins/Developer/Concert/ConcertSync/ConcertSyncTest/Source/ConcertSyncTest/Private/Replication/Util/Spec/ReplicationServer.cpp
@@ -19,10 +19,10 @@ namespace UE::ConcertSyncTests::Replication
 		, ServerReplicationManager(ConcertSyncServer::TestInterface::CreateServerReplicationManager(ServerSessionMock, *ReplicationWorkspace, InSessionFlags))
 	{}
 
-	FReplicationClient& FReplicationServer::ConnectClient()
+	FReplicationClient& FReplicationServer::ConnectClient(FConcertClientInfo ClientInfo)
 	{
 		const FGuid ClientEndpointId(0, 0, 0, Clients.Num() + 1); // {0, 0, 0, 0} is used by the server.
-		Clients.Add(MakeUnique<FReplicationClient>(ClientEndpointId, SessionFlags, *ServerSessionMock, TestContext));
+		Clients.Add(MakeUnique<FReplicationClient>(ClientEndpointId, SessionFlags, *ServerSessionMock, TestContext, MoveTemp(ClientInfo)));
 		ServerSessionMock->ConnectClient(ClientEndpointId, *(Clients.Last()->GetClientSessionMock()));
 		return *Clients.Last();
 	}
diff --git a/Engine/Plugins/Developer/Concert/ConcertSync/ConcertSyncTest/Source/ConcertSyncTest/Private/Replication/Util/Spec/ReplicationServer.h b/Engine/Plugins/Developer/Concert/ConcertSync/ConcertSyncTest/Source/ConcertSyncTest/Private/Replication/Util/Spec/ReplicationServer.h
--- a/Engine/Plugins/Developer/Concert/ConcertSync/ConcertSyncTest/Source/ConcertSyncTest/Private/Replication/Util/Spec/ReplicationServer.h
+++ b/Engine/Plugins/Developer/Concert/ConcertSync/ConcertSyncTest/Source/ConcertSyncTest/Private/Replication/Util/Spec/ReplicationServer.h
@@ -2,6 +2,7 @@
 
 #pragma once
 
+#include "ConcertMessageData.h"
 #include "ConcertSyncSessionFlags.h"
 #include "Replication/IConcertServerReplicationManager.h"
 #include "Replication/Util/Mocks/ReplicationWorkspaceEmptyMock.h"
@@ -9,8 +10,7 @@
 
 #include "Templates/UnrealTemplate.h"
 
-namespace UE::ConcertSyncServer::Replication { class IReplicationWorkspace;
-}
+namespace UE::ConcertSyncServer::Replication { class IReplicationWorkspace; }
 
 class FAutomationTestBase;
 
@@ -30,7 +30,7 @@ namespace UE::ConcertSyncTests::Replication
 			);
 
 		/** Connects a client to the server. */
-		FReplicationClient& ConnectClient();
+		FReplicationClient& ConnectClient(FConcertClientInfo ClientInfo = {});
 		
 		/** Lets the server process any messages that have come in. */
 		void TickServer(float FakeDeltaTime = 1.f / 60.f);
diff --git a/Engine/Plugins/Developer/Concert/ConcertSync/ConcertSyncTest/Source/ConcertSyncTest/Private/Util/ClientServerCommunicationTest.h b/Engine/Plugins/Developer/Concert/ConcertSync/ConcertSyncTest/Source/ConcertSyncTest/Private/Util/ClientServerCommunicationTest.h
--- a/Engine/Plugins/Developer/Concert/ConcertSync/ConcertSyncTest/Source/ConcertSyncTest/Private/Util/ClientServerCommunicationTest.h
+++ b/Engine/Plugins/Developer/Concert/ConcertSync/ConcertSyncTest/Source/ConcertSyncTest/Private/Util/ClientServerCommunicationTest.h
@@ -75,9 +75,10 @@ namespace UE::ConcertSyncTests
 	class FConcertClientSessionBaseMock : public IConcertClientSession
 	{
 	public:
-		FConcertClientSessionBaseMock() 
+		FConcertClientSessionBaseMock(FConcertClientInfo ClientInfo) 
 			: Id(FGuid::NewGuid()) 
-			, Name("FConcertClientSessionBaseMock") 
+			, Name("FConcertClientSessionBaseMock")
+			, ClientInfo(MoveTemp(ClientInfo))
 		{ }
 
 		// IConcertSession Begin.
@@ -110,7 +111,7 @@ namespace UE::ConcertSyncTests
 		virtual EConcertConnectionStatus GetConnectionStatus() const override            { return NotMocked<EConcertConnectionStatus>(EConcertConnectionStatus::Connected); }
 		virtual FGuid GetSessionClientEndpointId() const override                        { return NotMocked<FGuid>(); }
 		virtual FGuid GetSessionServerEndpointId() const override                        { return NotMocked<FGuid>(); }
-		virtual const FConcertClientInfo& GetLocalClientInfo() const override            { return NotMocked<const FConcertClientInfo&>(ClientInfo); }
+		virtual const FConcertClientInfo& GetLocalClientInfo() const override            { return ClientInfo; }
 		virtual void UpdateLocalClientInfo(const FConcertClientInfoUpdate&) override     { return NotMocked<void>(); }
 		virtual void Connect() override                                                  { return NotMocked<void>(); }
 		virtual void Disconnect() override                                               { return NotMocked<void>(); }
@@ -300,8 +301,18 @@ namespace UE::ConcertSyncTests
 
 		virtual bool FindSessionClient(const FGuid& EndpointId, FConcertSessionClientInfo& OutSessionClientInfo) const override
 		{
-			// Need to mock this in case test fails the internal server logic may log the client name which made an invalid request.
-			return true;
+			for (int32 i = 0; i < ClientSessions.Num(); ++i)
+			{
+				if (ClientEndpoints[i] != EndpointId)
+				{
+					continue;
+				}
+				
+				const FConcertClientSessionBaseMock* ClientSessionMock = ClientSessions[i];
+				OutSessionClientInfo = { EndpointId, ClientSessionMock->GetLocalClientInfo() };
+				return true;
+			}
+			return false;
 		}
 		
 		virtual FOnConcertServerSessionClientChanged& OnSessionClientChanged() override
@@ -345,8 +356,10 @@ namespace UE::ConcertSyncTests
 	class FConcertClientSessionMock : public FConcertClientSessionBaseMock
 	{
 	public:
-		FConcertClientSessionMock(const FGuid& ClientEndpointId, FConcertServerSessionMock& Server)
-			: ServerMock(Server)
+		
+		FConcertClientSessionMock(const FGuid& ClientEndpointId, FConcertServerSessionMock& Server, FConcertClientInfo ClientInfo)
+			: FConcertClientSessionBaseMock(MoveTemp(ClientInfo))
+			, ServerMock(Server)
 			, EndpointId(ClientEndpointId)
 		{
 		}
@@ -451,8 +464,8 @@ namespace UE::ConcertSyncTests
 	public:
 		struct FClientInfo
 		{
-			FClientInfo(const FGuid& ClientEndPointId, FConcertServerSessionMock& Server)
-				: ClientSessionMock(MakeShared<FConcertClientSessionMock>(ClientEndPointId, Server))
+			FClientInfo(const FGuid& ClientEndPointId, FConcertServerSessionMock& Server, FConcertClientInfo ClientInfo)
+				: ClientSessionMock(MakeShared<FConcertClientSessionMock>(ClientEndPointId, Server, MoveTemp(ClientInfo)))
 			{}
 
 			const TSharedRef<FConcertClientSessionBaseMock> ClientSessionMock;
@@ -471,10 +484,10 @@ namespace UE::ConcertSyncTests
 			Framework.OnTestEndEvent.RemoveAll(this);
 		}
 
-		FClientInfo& ConnectClient()
+		FClientInfo& ConnectClient(FConcertClientInfo ClientInfo = {})
 		{
 			const FGuid ClientEndpointId(0, 0, 0, Clients.Num() + 1); // {0, 0, 0, 0} is used by the server.
-			Clients.Add(MakeUnique<FClientInfo>(ClientEndpointId, *ServerSessionMock));
+			Clients.Add(MakeUnique<FClientInfo>(ClientEndpointId, *ServerSessionMock, MoveTemp(ClientInfo)));
 			ServerSessionMock->ConnectClient(ClientEndpointId, *(Clients.Last()->ClientSessionMock));
 			return *Clients.Last();
 		}
