Extract FScopedSessionDatabaseWithEndpoint out from FScopedSessionDatabase

diff --git a/Engine/Plugins/Developer/Concert/ConcertSync/ConcertSyncTest/Source/ConcertSyncTest/Private/HistoryEdition/DependencyGraphBuilderTests.cpp b/Engine/Plugins/Developer/Concert/ConcertSync/ConcertSyncTest/Source/ConcertSyncTest/Private/HistoryEdition/DependencyGraphBuilderTests.cpp
--- a/Engine/Plugins/Developer/Concert/ConcertSync/ConcertSyncTest/Source/ConcertSyncTest/Private/HistoryEdition/DependencyGraphBuilderTests.cpp
+++ b/Engine/Plugins/Developer/Concert/ConcertSync/ConcertSyncTest/Source/ConcertSyncTest/Private/HistoryEdition/DependencyGraphBuilderTests.cpp
@@ -9,8 +9,7 @@
 #include "HistoryTestUtil.h"
 #include "RenameEditAndDeleteMapsFlow.h"
 #include "Util/ActivityBuilder.h"
-#include "Util/ScopedSessionDatabase.h"
-
+#include "Util/ScopedSessionDatabaseWithEndpoint.h"
 
 namespace UE::ConcertSyncTests::RenameEditAndDeleteMapsFlowTest
 {
@@ -44,7 +43,7 @@ namespace UE::ConcertSyncTests::RenameEditAndDeleteMapsFlowTest
 	IMPLEMENT_SIMPLE_AUTOMATION_TEST(FRenameEditAndDeleteMapsFlowTest, "Editor.Concert.History.BuildGraph.RenameEditAndDeleteMapsFlow", EAutomationTestFlags::EditorContext | EAutomationTestFlags::EngineFilter)
 	bool FRenameEditAndDeleteMapsFlowTest::RunTest(const FString& Parameters)
 	{
-		FScopedSessionDatabase SessionDatabase(*this);
+		FScopedSessionDatabaseWithEndpoint SessionDatabase(*this);
 		const TTestActivityArray<int64> Activities = CreateActivityHistory(SessionDatabase, SessionDatabase.GetEndpoint());
 		
 		const ConcertSyncCore::FActivityDependencyGraph DependencyGraph = ConcertSyncCore::BuildDependencyGraphFrom(SessionDatabase);
@@ -180,7 +179,7 @@ namespace UE::ConcertSyncTests::DeletingAndRecreatingActorIsHardDependency
 		ActivityCount
 	};
 	
-	TArray<FActivityID> FillDatabase(FScopedSessionDatabase& Database);
+	TArray<FActivityID> FillDatabase(FScopedSessionDatabaseWithEndpoint& Database);
 	FConcertExportedObject CreateActorMetaData(FName OuterLevelPath);
 	
 	/**
@@ -195,7 +194,7 @@ namespace UE::ConcertSyncTests::DeletingAndRecreatingActorIsHardDependency
 	{
 		using namespace ConcertSyncCore;
 		
-		FScopedSessionDatabase SessionDatabase(*this);
+		FScopedSessionDatabaseWithEndpoint SessionDatabase(*this);
 		const TArray<FActivityID> TestActivities = FillDatabase(SessionDatabase);
 
 		const FActivityDependencyGraph DependencyGraph = BuildDependencyGraphFrom(SessionDatabase);
@@ -224,7 +223,7 @@ namespace UE::ConcertSyncTests::DeletingAndRecreatingActorIsHardDependency
 		return true;
 	}
 	
-	TArray<FActivityID> FillDatabase(FScopedSessionDatabase& SessionDatabase)
+	TArray<FActivityID> FillDatabase(FScopedSessionDatabaseWithEndpoint& SessionDatabase)
 	{
 		TArray<FActivityID> ActivityIDs;
 		ActivityIDs.SetNumUninitialized(ActivityCount);
@@ -346,7 +345,7 @@ namespace UE::ConcertSyncTests::PackageEditedDependencyTest
 		ActivityCount
 	};
 	
-	TArray<FActivityID> FillDatabase(FScopedSessionDatabase& SessionDatabase);
+	TArray<FActivityID> FillDatabase(FScopedSessionDatabaseWithEndpoint& SessionDatabase);
 	
 	/**
 	 * This tests that EActivityDependencyReason::PackageEdited dependencies are discovered correctly.
@@ -377,7 +376,7 @@ namespace UE::ConcertSyncTests::PackageEditedDependencyTest
 	{
 		using namespace ConcertSyncCore;
 		
-		FScopedSessionDatabase SessionDatabase(*this);
+		FScopedSessionDatabaseWithEndpoint SessionDatabase(*this);
 		TArray<FActivityID> Activities = FillDatabase(SessionDatabase);
 		
 		const FActivityDependencyGraph DependencyGraph = BuildDependencyGraphFrom(SessionDatabase);
@@ -465,7 +464,7 @@ namespace UE::ConcertSyncTests::PackageEditedDependencyTest
 		return true;
 	}
 
-	TArray<FActivityID> FillDatabase(FScopedSessionDatabase& SessionDatabase)
+	TArray<FActivityID> FillDatabase(FScopedSessionDatabaseWithEndpoint& SessionDatabase)
 	{
 		const FName FooLevel = TEXT("/Game/Foo");
 		const FName BarLevel = TEXT("/Game/Bar");
@@ -511,7 +510,7 @@ namespace UE::ConcertSyncTests::RenamingDependencyEdgeCases
 		ActivityCount
 	};
 	
-	TArray<FActivityID> FillDatabase(FScopedSessionDatabase& SessionDatabase);
+	TArray<FActivityID> FillDatabase(FScopedSessionDatabaseWithEndpoint& SessionDatabase);
 	
 	/**
 	 * This tests that renaming dependencies are correctly set in some edge cases.
@@ -532,7 +531,7 @@ namespace UE::ConcertSyncTests::RenamingDependencyEdgeCases
 	{
 		using namespace ConcertSyncCore;
 		
-		FScopedSessionDatabase SessionDatabase(*this);
+		FScopedSessionDatabaseWithEndpoint SessionDatabase(*this);
 		TArray<FActivityID> Activities = FillDatabase(SessionDatabase);
 		
 		const FActivityDependencyGraph DependencyGraph = BuildDependencyGraphFrom(SessionDatabase);
@@ -556,7 +555,7 @@ namespace UE::ConcertSyncTests::RenamingDependencyEdgeCases
 		return true;
 	}
 	
-	TArray<FActivityID> FillDatabase(FScopedSessionDatabase& SessionDatabase)
+	TArray<FActivityID> FillDatabase(FScopedSessionDatabaseWithEndpoint& SessionDatabase)
 	{
 		const FName FooLevel = TEXT("/Game/Foo");
 		const FName BarLevel = TEXT("/Game/Bar");
@@ -596,14 +595,14 @@ namespace UE::ConcertSyncTests::TransactionDependencyEdgeCases
 		ActivityCount
 	};
 	
-	TArray<FActivityID> FillDatabase(FScopedSessionDatabase& SessionDatabase);
+	TArray<FActivityID> FillDatabase(FScopedSessionDatabaseWithEndpoint& SessionDatabase);
 	
 	IMPLEMENT_SIMPLE_AUTOMATION_TEST(FTransactionDependencyEdgeCaseTests, "Editor.Concert.History.BuildGraph.TransactionDependencyEdgeCaseTests", EAutomationTestFlags::EditorContext | EAutomationTestFlags::EngineFilter)
 	bool FTransactionDependencyEdgeCaseTests::RunTest(const FString& Parameters)
 	{
 		using namespace ConcertSyncCore;
 		
-		FScopedSessionDatabase SessionDatabase(*this);
+		FScopedSessionDatabaseWithEndpoint SessionDatabase(*this);
 		TArray<FActivityID> Activities = FillDatabase(SessionDatabase);
 		
 		const FActivityDependencyGraph DependencyGraph = BuildDependencyGraphFrom(SessionDatabase);
@@ -623,7 +622,7 @@ namespace UE::ConcertSyncTests::TransactionDependencyEdgeCases
 		return true;
 	}
 
-	TArray<FActivityID> FillDatabase(FScopedSessionDatabase& SessionDatabase)
+	TArray<FActivityID> FillDatabase(FScopedSessionDatabaseWithEndpoint& SessionDatabase)
 	{
 		const FName FooLevel = TEXT("/Game/Foo");
 		const FName BarLevel = TEXT("/Game/Bar");
diff --git a/Engine/Plugins/Developer/Concert/ConcertSync/ConcertSyncTest/Source/ConcertSyncTest/Private/HistoryEdition/HistoryBottomUpAnalysisTests.cpp b/Engine/Plugins/Developer/Concert/ConcertSync/ConcertSyncTest/Source/ConcertSyncTest/Private/HistoryEdition/HistoryBottomUpAnalysisTests.cpp
--- a/Engine/Plugins/Developer/Concert/ConcertSync/ConcertSyncTest/Source/ConcertSyncTest/Private/HistoryEdition/HistoryBottomUpAnalysisTests.cpp
+++ b/Engine/Plugins/Developer/Concert/ConcertSync/ConcertSyncTest/Source/ConcertSyncTest/Private/HistoryEdition/HistoryBottomUpAnalysisTests.cpp
@@ -8,10 +8,8 @@
 #include "HistoryEdition/DependencyGraphBuilder.h"
 #include "HistoryEdition/ActivityGraphIDs.h"
 #include "HistoryEdition/HistoryAnalysis.h"
-#include "HistoryEdition/ActivityNode.h"
 #include "RenameEditAndDeleteMapsFlow.h"
-#include "Util/ScopedSessionDatabase.h"
-
+#include "Util/ScopedSessionDatabaseWithEndpoint.h"
 
 namespace UE::ConcertSyncTests::AnalysisTests
 {
@@ -24,7 +22,7 @@ namespace UE::ConcertSyncTests::AnalysisTests
 		using namespace RenameEditAndDeleteMapsFlowTest;
 		using namespace ConcertSyncCore;
 		
-		FScopedSessionDatabase SessionDatabase(*this);
+		FScopedSessionDatabaseWithEndpoint SessionDatabase(*this);
 		const TTestActivityArray<FActivityID> Activities = CreateActivityHistory(SessionDatabase, SessionDatabase.GetEndpoint());
 		const FActivityDependencyGraph DependencyGraph = BuildDependencyGraphFrom(SessionDatabase);
 		UE_LOG(LogConcert, Log, TEXT("%s tested graph in Graphviz format:\n\n%s"), *GetTestFullName(), *ConcertSyncCore::Graphviz::ExportToGraphviz(DependencyGraph, SessionDatabase));
diff --git a/Engine/Plugins/Developer/Concert/ConcertSync/ConcertSyncTest/Source/ConcertSyncTest/Private/HistoryEdition/HistoryTopDownAnalysisTests.cpp b/Engine/Plugins/Developer/Concert/ConcertSync/ConcertSyncTest/Source/ConcertSyncTest/Private/HistoryEdition/HistoryTopDownAnalysisTests.cpp
--- a/Engine/Plugins/Developer/Concert/ConcertSync/ConcertSyncTest/Source/ConcertSyncTest/Private/HistoryEdition/HistoryTopDownAnalysisTests.cpp
+++ b/Engine/Plugins/Developer/Concert/ConcertSync/ConcertSyncTest/Source/ConcertSyncTest/Private/HistoryEdition/HistoryTopDownAnalysisTests.cpp
@@ -8,9 +8,8 @@
 #include "HistoryEdition/DependencyGraphBuilder.h"
 #include "HistoryEdition/ActivityGraphIDs.h"
 #include "HistoryEdition/HistoryAnalysis.h"
-#include "HistoryEdition/ActivityNode.h"
 #include "RenameEditAndDeleteMapsFlow.h"
-#include "Util/ScopedSessionDatabase.h"
+#include "Util/ScopedSessionDatabaseWithEndpoint.h"
 
 #include "Algo/AllOf.h"
 
@@ -64,7 +63,7 @@ namespace UE::ConcertSyncTests::AnalysisTests
 		using namespace RenameEditAndDeleteMapsFlowTest;
 		using namespace ConcertSyncCore;
 		
-		FScopedSessionDatabase SessionDatabase(*this);
+		FScopedSessionDatabaseWithEndpoint SessionDatabase(*this);
 		const TTestActivityArray<FActivityID> Activities = CreateActivityHistory(SessionDatabase, SessionDatabase.GetEndpoint());
 		const FActivityDependencyGraph DependencyGraph = BuildDependencyGraphFrom(SessionDatabase);
 		UE_LOG(LogConcert, Log, TEXT("%s tested graph in Graphviz format:\n\n%s"), *GetTestFullName(), *ConcertSyncCore::Graphviz::ExportToGraphviz(DependencyGraph, SessionDatabase));
diff --git a/Engine/Plugins/Developer/Concert/ConcertSync/ConcertSyncTest/Source/ConcertSyncTest/Private/Util/ActivityBuilder.h b/Engine/Plugins/Developer/Concert/ConcertSync/ConcertSyncTest/Source/ConcertSyncTest/Private/Util/ActivityBuilder.h
--- a/Engine/Plugins/Developer/Concert/ConcertSync/ConcertSyncTest/Source/ConcertSyncTest/Private/Util/ActivityBuilder.h
+++ b/Engine/Plugins/Developer/Concert/ConcertSync/ConcertSyncTest/Source/ConcertSyncTest/Private/Util/ActivityBuilder.h
@@ -2,7 +2,7 @@
 
 #pragma once
 
-#include "ScopedSessionDatabase.h"
+#include "ScopedSessionDatabaseWithEndpoint.h"
 
 namespace UE::ConcertSyncTests
 {
@@ -12,10 +12,10 @@ namespace UE::ConcertSyncTests
 	class FActivityBuilder
 	{
 		TArray<FActivityID> Activities;
-		FScopedSessionDatabase& SessionDatabase;
+		FScopedSessionDatabaseWithEndpoint& SessionDatabase;
 	public:
 
-		FActivityBuilder(FScopedSessionDatabase& SessionDatabase, uint32 ActivityCount)
+		FActivityBuilder(FScopedSessionDatabaseWithEndpoint& SessionDatabase, uint32 ActivityCount)
 			: SessionDatabase(SessionDatabase)
 		{
 			Activities.SetNumUninitialized(ActivityCount);
diff --git a/Engine/Plugins/Developer/Concert/ConcertSync/ConcertSyncTest/Source/ConcertSyncTest/Private/Util/ScopedSessionDatabase.h b/Engine/Plugins/Developer/Concert/ConcertSync/ConcertSyncTest/Source/ConcertSyncTest/Private/Util/ScopedSessionDatabase.h
--- a/Engine/Plugins/Developer/Concert/ConcertSync/ConcertSyncTest/Source/ConcertSyncTest/Private/Util/ScopedSessionDatabase.h
+++ b/Engine/Plugins/Developer/Concert/ConcertSync/ConcertSyncTest/Source/ConcertSyncTest/Private/Util/ScopedSessionDatabase.h
@@ -15,21 +15,12 @@ namespace UE::ConcertSyncTests
 	{
 		FAutomationTestBase& Test;
 		const FString TestSessionPath_Server = FPaths::ProjectIntermediateDir() / TEXT("ConcertDatabaseTest_Server");
-		const FGuid EndpointID;
 	public:
 
 		FScopedSessionDatabase(FAutomationTestBase& Test)
 			: Test(Test)
-			, EndpointID(FGuid::NewGuid())
 		{
 			Open(TestSessionPath_Server);
-
-			FConcertSyncEndpointData EndpointData;
-			EndpointData.ClientInfo.Initialize();
-			if (!SetEndpoint(EndpointID, EndpointData))
-			{
-				Test.AddError(FString::Printf(TEXT("Test may be faulty because endpoint could not be set: %s"), *GetLastError()));
-			}
 		}
 
 		~FScopedSessionDatabase()
@@ -40,7 +31,5 @@ namespace UE::ConcertSyncTests
 			}
 			IFileManager::Get().DeleteDirectory(*TestSessionPath_Server, false, true);
 		}
-
-		const FGuid& GetEndpoint() const { return EndpointID; }
 	};
 }
diff --git a/Engine/Plugins/Developer/Concert/ConcertSync/ConcertSyncTest/Source/ConcertSyncTest/Private/Util/ScopedSessionDatabaseWithEndpoint.h b/Engine/Plugins/Developer/Concert/ConcertSync/ConcertSyncTest/Source/ConcertSyncTest/Private/Util/ScopedSessionDatabaseWithEndpoint.h
new file mode 100644
--- /dev/null
+++ b/Engine/Plugins/Developer/Concert/ConcertSync/ConcertSyncTest/Source/ConcertSyncTest/Private/Util/ScopedSessionDatabaseWithEndpoint.h
@@ -0,0 +1,28 @@
+// Copyright Epic Games, Inc. All Rights Reserved.
+
+#pragma once
+
+#include "ScopedSessionDatabase.h"
+
+namespace UE::ConcertSyncTests
+{
+	class FScopedSessionDatabaseWithEndpoint : public FScopedSessionDatabase
+	{
+		const FGuid EndpointID;
+	public:
+
+		FScopedSessionDatabaseWithEndpoint(FAutomationTestBase& Test)
+			: FScopedSessionDatabase(Test)
+			, EndpointID(FGuid::NewGuid())
+		{
+			FConcertSyncEndpointData EndpointData;
+			EndpointData.ClientInfo.Initialize();
+			if (!SetEndpoint(EndpointID, EndpointData))
+			{
+				Test.AddError(FString::Printf(TEXT("Test may be faulty because endpoint could not be set: %s"), *GetLastError()));
+			}
+		}
+
+		const FGuid& GetEndpoint() const { return EndpointID; }
+	};
+}
