Fully implement and unit test IReplicationWorkspace

diff --git a/Engine/Plugins/Developer/Concert/ConcertSync/ConcertSyncServer/Source/ConcertSyncServer/Private/ConcertServerWorkspace.cpp b/Engine/Plugins/Developer/Concert/ConcertSync/ConcertSyncServer/Source/ConcertSyncServer/Private/ConcertServerWorkspace.cpp
--- a/Engine/Plugins/Developer/Concert/ConcertSync/ConcertSyncServer/Source/ConcertSyncServer/Private/ConcertServerWorkspace.cpp
+++ b/Engine/Plugins/Developer/Concert/ConcertSync/ConcertSyncServer/Source/ConcertSyncServer/Private/ConcertServerWorkspace.cpp
@@ -13,7 +13,6 @@
 #include "ConcertLogGlobal.h"
 #include "ConcertPackageEvents.h"
 #include "ConcertUtil.h"
-#include "HistoryEdition/ActivityDependencyGraph.h"
 #include "Replication/Messages/ReplicationActivity.h"
 #include "Serialization/MemoryReader.h"
 
@@ -21,8 +20,19 @@
 
 FConcertServerWorkspace::FConcertServerWorkspace(const TSharedRef<FConcertSyncServerLiveSession>& InLiveSession, TSharedPtr<IConcertFileSharingService> InFileSharingService)
 	: FileSharingService(MoveTemp(InFileSharingService))
+	, ReplicationWorkspace(
+		InLiveSession->GetSessionDatabase(),
+		UE::ConcertSyncServer::FFindSessionClient::CreateLambda([this](const FGuid& EndpointId)
+		{
+			FConcertSessionClientInfo Info;
+			return LiveSession->GetSession().FindSessionClient(EndpointId, Info) ? Info : TOptional<FConcertSessionClientInfo>{};
+		}),
+		UE::ConcertSyncServer::FShouldIgnoreClientActivityOnRestore::CreateRaw(this, &FConcertServerWorkspace::ShouldIgnoreClientActivityOnRestore)
+		)
 {
 	BindSession(InLiveSession);
+	// Sync replication activity back to clients when it is added to the database
+	ReplicationWorkspace.OnAddReplicationActivity().AddRaw(this, &FConcertServerWorkspace::OnAddReplicationActivity);
 }
 
 FConcertServerWorkspace::~FConcertServerWorkspace()
@@ -30,31 +40,6 @@ FConcertServerWorkspace::~FConcertServerWorkspace()
 	UnbindSession();
 }
 
-void FConcertServerWorkspace::ProduceClientLeaveReplicationActivity(const FGuid& EndpointId, const FConcertSyncReplicationPayload_LeaveReplication& EventData)
-{
-	if (EnumHasAnyFlags(LiveSession->GetSessionFlags(), EConcertSyncSessionFlags::ShouldEnableReplicationActivities))
-	{
-		FConcertSyncReplicationActivity Activity;
-		Activity.EndpointId = EndpointId;
-		Activity.EventData.SetPayload(EventData);
-		Activity.EventSummary.SetTypedPayload(FConcertSyncReplicationActivitySummary::CreateSummaryForEvent(Activity.EventData));
-		Activity.bIgnored = ShouldIgnoreClientActivityOnRestore(EndpointId);
-		AddReplicationActivity(Activity);
-	}
-}
-
-bool FConcertServerWorkspace::GetLastLeaveReplicationActivityByClient(const FConcertClientInfo& InClientInfo, FConcertSyncReplicationPayload_LeaveReplication& OutLeaveReplication) const
-{
-	// TODO: DO
-	return false;
-}
-
-bool FConcertServerWorkspace::GetLeaveReplicationActivityById(const int64 ActivityId, FConcertSyncReplicationPayload_LeaveReplication& OutLeaveReplication) const
-{
-	// TODO: DO
-	return false;
-}
-
 void FConcertServerWorkspace::BindSession(const TSharedRef<FConcertSyncServerLiveSession>& InLiveSession)
 {
 	check(InLiveSession->IsValidSession());
@@ -1181,12 +1166,9 @@ void FConcertServerWorkspace::SendSyncPackageActivityEvent(const FConcertWorkspa
 	LiveSession->GetSession().SendCustomEvent(SyncEvent, InTargetEndpointId, EConcertMessageFlags::ReliableOrdered);
 }
 
-void FConcertServerWorkspace::AddReplicationActivity(const FConcertSyncReplicationActivity& InReplicationActivity)
+void FConcertServerWorkspace::OnAddReplicationActivity(const int64 ActivityId, const bool bSuccess)
 {
-	// Add the activity and sync it
-	int64 ActivityId = 0;
-	int64 EventId = 0;
-	if (LiveSession->GetSessionDatabase().AddReplicationActivity(InReplicationActivity, ActivityId, EventId))
+	if (bSuccess)
 	{
 		PostActivityAdded(ActivityId);
 		SyncCommandQueue->QueueCommand(LiveSyncEndpoints, [this, SyncActivityId = ActivityId](const FConcertServerSyncCommandQueue::FSyncCommandContext& InSyncCommandContext, const FGuid& InEndpointId)
diff --git a/Engine/Plugins/Developer/Concert/ConcertSync/ConcertSyncServer/Source/ConcertSyncServer/Private/ConcertServerWorkspace.h b/Engine/Plugins/Developer/Concert/ConcertSync/ConcertSyncServer/Source/ConcertSyncServer/Private/ConcertServerWorkspace.h
--- a/Engine/Plugins/Developer/Concert/ConcertSync/ConcertSyncServer/Source/ConcertSyncServer/Private/ConcertServerWorkspace.h
+++ b/Engine/Plugins/Developer/Concert/ConcertSync/ConcertSyncServer/Source/ConcertSyncServer/Private/ConcertServerWorkspace.h
@@ -7,6 +7,7 @@
 #include "IConcertSessionHandler.h"
 #include "Misc/Optional.h"
 #include "Replication/IReplicationWorkspace.h"
+#include "Replication/ReplicationWorkspace.h"
 
 class IConcertFileSharingService;
 class IConcertMessage;
@@ -31,18 +32,15 @@ enum class EConcertLockFlags : uint8
 };
 ENUM_CLASS_FLAGS(EConcertLockFlags);
 
-class FConcertServerWorkspace : public UE::ConcertSyncServer::Replication::IReplicationWorkspace
+class FConcertServerWorkspace
 {
 public:
 	
 	explicit FConcertServerWorkspace(const TSharedRef<FConcertSyncServerLiveSession>& InLiveSession, TSharedPtr<IConcertFileSharingService> InFileSharingService);
-	virtual ~FConcertServerWorkspace() override;
+	~FConcertServerWorkspace();
 
-	//~ Begin IReplicationWorkspace Interface
-	virtual void ProduceClientLeaveReplicationActivity(const FGuid& EndpointId, const FConcertSyncReplicationPayload_LeaveReplication& EventData) override;
-	virtual bool GetLastLeaveReplicationActivityByClient(const FConcertClientInfo& InClientInfo, FConcertSyncReplicationPayload_LeaveReplication& OutLeaveReplication) const override;
-	virtual bool GetLeaveReplicationActivityById(const int64 ActivityId, FConcertSyncReplicationPayload_LeaveReplication& OutLeaveReplication) const override;
-	//~ End IReplicationWorkspace Interface
+	/** @return Gets the portion of the workspace that manages replication data. */
+	UE::ConcertSyncServer::Replication::IReplicationWorkspace& GetReplicationWorkspace() { return ReplicationWorkspace; }
 
 private:
 	
@@ -261,12 +259,13 @@ private:
 	void SendSyncPackageActivityEvent(const FConcertWorkspaceSyncActivityEvent& SyncEvent, const FGuid& InTargetEndpointId) const;
 
 	/**
-	 * Add a new replication activity to the session database, and sync the result back to all clients.
-	 * @note The endpoint ID referenced by the activity must exist in the database (@see SetEndpoint).
+	 * Called when the replication workspace attempted to add a FConcertSyncReplicationActivity to the database.
+	 * This syncs the result back to all clients.
 	 *
-	 * @param InReplicationActivity		The replication activity to add (the ActivityId, EventTime, EventType, and EventId members are ignored).
+	 * @param ActivityId	The identifier of the activity that was attempted to be added
+	 * @param bSuccess		Whether the activity was actually added to the database (i.e. false means the attempt failed).
 	 */
-	void AddReplicationActivity(const FConcertSyncReplicationActivity& InReplicationActivity);
+	void OnAddReplicationActivity(const int64 ActivityId, const bool bSuccess);
 	
 	/**
 	 * Send a sync event for a replication activity in the session database.
@@ -289,7 +288,7 @@ private:
 	 * @param ClientEndpoint The client for which the check must be done.
 	 * @return True if the activity should be flagged as 'ignored on restore'.
 	 */
-	bool ShouldIgnoreClientActivityOnRestore(const FGuid ClientEndpoint) const { return IgnoredActivityClients.Contains(ClientEndpoint); }
+	bool ShouldIgnoreClientActivityOnRestore(const FGuid& ClientEndpoint) const { return IgnoredActivityClients.Contains(ClientEndpoint); }
 
 	/**
 	 * Returns true if the package data is small enough to be exchanged using a single TArray<> data structure. If the package data is
@@ -349,4 +348,7 @@ private:
 
 	/** Optional side channel to exchange large blobs (package data) with the client in a scalable way (ex. the request/response transport layer is not designed and doesn't support exchanging 3GB packages). */
 	TSharedPtr<IConcertFileSharingService> FileSharingService;
+
+	/** Manages replication data for the database. */
+	UE::ConcertSyncServer::FReplicationWorkspace ReplicationWorkspace;
 };
diff --git a/Engine/Plugins/Developer/Concert/ConcertSync/ConcertSyncServer/Source/ConcertSyncServer/Private/ConcertSyncServer.cpp b/Engine/Plugins/Developer/Concert/ConcertSync/ConcertSyncServer/Source/ConcertSyncServer/Private/ConcertSyncServer.cpp
--- a/Engine/Plugins/Developer/Concert/ConcertSync/ConcertSyncServer/Source/ConcertSyncServer/Private/ConcertSyncServer.cpp
+++ b/Engine/Plugins/Developer/Concert/ConcertSync/ConcertSyncServer/Source/ConcertSyncServer/Private/ConcertSyncServer.cpp
@@ -580,7 +580,7 @@ bool FConcertSyncServer::CreateLiveSession(const TSharedRef<IConcertServerSessio
 		// Create Replication Manager
 		if (EnumHasAnyFlags(LiveSession->GetSessionFlags(), EConcertSyncSessionFlags::EnableReplication))
 		{
-			CreateReplicationManager(InSession, *Workspace, LiveSession->GetSessionFlags());
+			CreateReplicationManager(InSession, Workspace->GetReplicationWorkspace(), LiveSession->GetSessionFlags());
 		}
 
 		// We needn't call OnActivityProduced().Remove(...) because the subscription needs to stay for the lifetime of FConcertSyncServerLiveSession::SessionDatabase
diff --git a/Engine/Plugins/Developer/Concert/ConcertSync/ConcertSyncServer/Source/ConcertSyncServer/Private/ConcertSyncServer.h b/Engine/Plugins/Developer/Concert/ConcertSync/ConcertSyncServer/Source/ConcertSyncServer/Private/ConcertSyncServer.h
--- a/Engine/Plugins/Developer/Concert/ConcertSync/ConcertSyncServer/Source/ConcertSyncServer/Private/ConcertSyncServer.h
+++ b/Engine/Plugins/Developer/Concert/ConcertSync/ConcertSyncServer/Source/ConcertSyncServer/Private/ConcertSyncServer.h
@@ -6,6 +6,7 @@
 #include "IConcertSyncServer.h"
 #include "IConcertServerEventSink.h"
 #include "HAL/IConsoleManager.h"
+#include "HAL/Platform.h"
 
 class IConcertServerSession;
 class FConcertServerWorkspace;
@@ -63,7 +64,7 @@ private:
 	void CreateSequencerManager(const TSharedRef<FConcertSyncServerLiveSession>& InLiveSession);
 	void DestroySequencerManager(const TSharedRef<FConcertSyncServerLiveSession>& InLiveSession);
 	
-	void CreateReplicationManager(const TSharedRef<IConcertServerSession>& InSession, UE::ConcertSyncServer::Replication::IReplicationWorkspace& InWorkspace, EConcertSyncSessionFlags InSessionFlags);
+	void CreateReplicationManager(const TSharedRef<IConcertServerSession>& InSession, UE::ConcertSyncServer::Replication::IReplicationWorkspace& InWorkspace UE_LIFETIMEBOUND, EConcertSyncSessionFlags InSessionFlags);
 	void DestroyReplicationManager(const FGuid& SessionId); 
 
 	bool CreateLiveSession(const TSharedRef<IConcertServerSession>& InSession, const FInternalLiveSessionCreationParams& AdditionalParams);
diff --git a/Engine/Plugins/Developer/Concert/ConcertSync/ConcertSyncServer/Source/ConcertSyncServer/Private/Replication/ConcertServerReplicationManager_RestoreContent.cpp b/Engine/Plugins/Developer/Concert/ConcertSync/ConcertSyncServer/Source/ConcertSyncServer/Private/Replication/ConcertServerReplicationManager_RestoreContent.cpp
--- a/Engine/Plugins/Developer/Concert/ConcertSync/ConcertSyncServer/Source/ConcertSyncServer/Private/Replication/ConcertServerReplicationManager_RestoreContent.cpp
+++ b/Engine/Plugins/Developer/Concert/ConcertSync/ConcertSyncServer/Source/ConcertSyncServer/Private/Replication/ConcertServerReplicationManager_RestoreContent.cpp
@@ -244,7 +244,7 @@ namespace UE::ConcertSyncServer::Replication
 		const bool bOverrideActivity = Request.ActivityId.IsSet();
 		const bool bFound = bOverrideActivity
 			? ServerWorkspace.GetLeaveReplicationActivityById(*Request.ActivityId, OldClientState)
-			: ServerWorkspace.GetLastLeaveReplicationActivityByClient(ClientInfo.ClientInfo, OldClientState);
+			: ServerWorkspace.GetLastLeaveReplicationActivityByClient(ClientInfo, OldClientState);
 		if (!bFound)
 		{
 			return MakeTuple(
diff --git a/Engine/Plugins/Developer/Concert/ConcertSync/ConcertSyncServer/Source/ConcertSyncServer/Private/Replication/ReplicationTestInterface.cpp b/Engine/Plugins/Developer/Concert/ConcertSync/ConcertSyncServer/Source/ConcertSyncServer/Private/Replication/ReplicationTestInterface.cpp
--- a/Engine/Plugins/Developer/Concert/ConcertSync/ConcertSyncServer/Source/ConcertSyncServer/Private/Replication/ReplicationTestInterface.cpp
+++ b/Engine/Plugins/Developer/Concert/ConcertSync/ConcertSyncServer/Source/ConcertSyncServer/Private/Replication/ReplicationTestInterface.cpp
@@ -1,8 +1,13 @@
 // Copyright Epic Games, Inc. All Rights Reserved.
 
 #include "ConcertServerReplicationManager.h"
+#include "ReplicationWorkspace.h"
 #include "Replication/IConcertServerReplicationManager.h"
 
+#include "Misc/Guid.h"
+#include "Misc/Optional.h"
+#include "Templates/Function.h"
+
 class IConcertServerSession;
 
 /**
@@ -21,4 +26,19 @@ namespace UE::ConcertSyncServer::TestInterface
 	{
 		return MakeShared<Replication::FConcertServerReplicationManager>(InLiveSession, InWorkspace, InSessionFlags);
 	}
+
+	CONCERTSYNCSERVER_API TSharedRef<Replication::IReplicationWorkspace> CreateReplicationWorkspace(
+		FConcertSyncSessionDatabase& Database,
+		TFunction<TOptional<FConcertSessionClientInfo>(const FGuid& EndpointId)> FindSessionClient,
+		TFunction<bool(const FGuid& ClientId)> ShouldIgnoreClientActivityOnRestore
+		)
+	{
+		return MakeShared<FReplicationWorkspace>(
+			Database,
+			FFindSessionClient::CreateLambda([FindSessionClient = MoveTemp(FindSessionClient)](const FGuid& EndpointId)
+				{ return FindSessionClient(EndpointId); }),
+			FShouldIgnoreClientActivityOnRestore::CreateLambda([ShouldIgnoreClientActivityOnRestore = MoveTemp(ShouldIgnoreClientActivityOnRestore)](const FGuid& EndpointId)
+				{ return ShouldIgnoreClientActivityOnRestore(EndpointId); })
+			);
+	}
 }
diff --git a/Engine/Plugins/Developer/Concert/ConcertSync/ConcertSyncServer/Source/ConcertSyncServer/Private/Replication/ReplicationWorkspace.cpp b/Engine/Plugins/Developer/Concert/ConcertSync/ConcertSyncServer/Source/ConcertSyncServer/Private/Replication/ReplicationWorkspace.cpp
new file mode 100644
--- /dev/null
+++ b/Engine/Plugins/Developer/Concert/ConcertSync/ConcertSyncServer/Source/ConcertSyncServer/Private/Replication/ReplicationWorkspace.cpp
@@ -0,0 +1,84 @@
+﻿// Copyright Epic Games, Inc. All Rights Reserved.
+
+#include "ReplicationWorkspace.h"
+
+#include "ConcertSyncSessionDatabase.h"
+#include "Replication/Messages/ReplicationActivity.h"
+
+namespace UE::ConcertSyncServer
+{
+	FReplicationWorkspace::FReplicationWorkspace(
+		FConcertSyncSessionDatabase& Database,
+		FFindSessionClient InFindSessionClientDelegate,
+		FShouldIgnoreClientActivityOnRestore InShouldIgnoreClientActivityOnRestoreDelegate
+		)
+		: Database(Database)
+		, FindSessionClientDelegate(MoveTemp(InFindSessionClientDelegate))
+		, ShouldIgnoreClientActivityOnRestoreDelegate(MoveTemp(InShouldIgnoreClientActivityOnRestoreDelegate))
+	{
+		check(FindSessionClientDelegate.IsBound() && ShouldIgnoreClientActivityOnRestoreDelegate.IsBound());
+	}
+
+	TOptional<int64> FReplicationWorkspace::ProduceClientLeaveReplicationActivity(const FGuid& EndpointId, const FConcertSyncReplicationPayload_LeaveReplication& EventData)
+	{
+		FConcertSyncReplicationActivity Activity;
+		Activity.EndpointId = EndpointId;
+		Activity.EventData.SetPayload(EventData);
+		Activity.EventSummary.SetTypedPayload(FConcertSyncReplicationActivitySummary::CreateSummaryForEvent(Activity.EventData));
+		Activity.bIgnored = ShouldIgnoreClientActivityOnRestoreDelegate.Execute(EndpointId);
+		
+		int64 ActivityId = 0;
+		int64 EventId = 0;
+		const bool bSuccess = Database.AddReplicationActivity(Activity, ActivityId, EventId);
+		// This allows FConcertServerWorkspace to send the activity to other clients
+		OnAddReplicationActivityDelegate.Broadcast(ActivityId, bSuccess);
+
+		return bSuccess ? ActivityId : TOptional<int64>{};
+	}
+
+	bool FReplicationWorkspace::GetLastLeaveReplicationActivityByClient(const FConcertSessionClientInfo& InClientInfo, FConcertSyncReplicationPayload_LeaveReplication& OutLeaveReplication) const
+	{
+		// Goal: Iterate all endpoints with the same display name. Then, get the latest transaction ID from all of them.
+		// This entire approach is suboptimal though: There should be a dedicated SQL prepared statement  to which we give the client and device name.
+		// That's not possible though because the endpoints.client_info_data cannot be used to query by SQL as it's a BLOB; this strictly also violates
+		// column atomicity of 1NF.
+		TArray<FConcertSyncEndpointIdAndData> EndpointsWithSameDeviceName;
+		TArray<FConcertSyncEndpointIdAndData> EndpointsWithOtherDeviceName;
+		Database.EnumerateEndpoints([&InClientInfo, &EndpointsWithSameDeviceName, &EndpointsWithOtherDeviceName](FConcertSyncEndpointIdAndData&& EndpointIdAndData)
+		{
+			const FConcertClientInfo& ClientInfo = EndpointIdAndData.EndpointData.ClientInfo;
+			const FConcertClientInfo& SearchedClientInfo = InClientInfo.ClientInfo;
+			if (ClientInfo.DisplayName == SearchedClientInfo.DisplayName)
+			{
+				const bool bSameDeviceName = ClientInfo.DeviceName == SearchedClientInfo.DeviceName;
+				TArray<FConcertSyncEndpointIdAndData>& ArrayToAddTo = bSameDeviceName ? EndpointsWithSameDeviceName : EndpointsWithOtherDeviceName;
+				ArrayToAddTo.Add(EndpointIdAndData);
+			}
+			return true;
+		});
+
+		const TArray<FConcertSyncEndpointIdAndData>& ArrayToSearch = EndpointsWithSameDeviceName.IsEmpty()
+			? EndpointsWithOtherDeviceName
+			: EndpointsWithSameDeviceName;
+		int64 NewestEventId = INDEX_NONE;
+		for (const FConcertSyncEndpointIdAndData& ClientData : ArrayToSearch)
+		{
+			int64 EventId = INDEX_NONE;
+			Database.GetReplicationMaxEventIdByClientAndType(ClientData.EndpointId, EConcertSyncReplicationActivityType::LeaveReplication, EventId);
+			NewestEventId = FMath::Max(NewestEventId, EventId);
+		}
+
+		FConcertSyncReplicationEvent Event;
+		const bool bGotData = NewestEventId != INDEX_NONE && Database.GetReplicationEvent(NewestEventId, Event);
+		return bGotData
+			&& Event.ActivityType == EConcertSyncReplicationActivityType::LeaveReplication
+			&& Event.GetPayload(OutLeaveReplication);
+	}
+
+	bool FReplicationWorkspace::GetLeaveReplicationActivityById(const int64 ActivityId, FConcertSyncReplicationPayload_LeaveReplication& OutLeaveReplication) const
+	{
+		FConcertSyncReplicationEvent Event;
+		const bool bGotEventData = Database.GetReplicationEvent(ActivityId, Event);
+		return bGotEventData && Event.ActivityType == EConcertSyncReplicationActivityType::LeaveReplication && Event.GetPayload(OutLeaveReplication);
+	}
+}
diff --git a/Engine/Plugins/Developer/Concert/ConcertSync/ConcertSyncServer/Source/ConcertSyncServer/Private/Replication/ReplicationWorkspace.h b/Engine/Plugins/Developer/Concert/ConcertSync/ConcertSyncServer/Source/ConcertSyncServer/Private/Replication/ReplicationWorkspace.h
new file mode 100644
--- /dev/null
+++ b/Engine/Plugins/Developer/Concert/ConcertSync/ConcertSyncServer/Source/ConcertSyncServer/Private/Replication/ReplicationWorkspace.h
@@ -0,0 +1,56 @@
+﻿// Copyright Epic Games, Inc. All Rights Reserved.
+
+#pragma once
+
+#include "Replication/IReplicationWorkspace.h"
+
+#include "Delegates/Delegate.h"
+#include "HAL/Platform.h"
+#include "Misc/Optional.h"
+
+class FConcertSyncSessionDatabase;
+
+namespace UE::ConcertSyncServer
+{
+	DECLARE_DELEGATE_RetVal_OneParam(TOptional<FConcertSessionClientInfo>, FFindSessionClient, const FGuid& EndpointId);
+	DECLARE_DELEGATE_RetVal_OneParam(bool, FShouldIgnoreClientActivityOnRestore, const FGuid& ClientId);
+	
+	/**
+	 * Implements the replication workspace server-side.
+	 * 
+	 * At runtime, this is created by FConcertServerWorkspace.
+	 * This exists as independent class so it can be unit-tested.
+	 */
+	class FReplicationWorkspace : public Replication::IReplicationWorkspace
+	{
+	public:
+
+		FReplicationWorkspace(
+			FConcertSyncSessionDatabase& Database UE_LIFETIMEBOUND,
+			FFindSessionClient InFindSessionClientDelegate,
+			FShouldIgnoreClientActivityOnRestore InShouldIgnoreClientActivityOnRestoreDelegate
+			);
+
+		/** Called when an activity is added to the database. */
+		DECLARE_MULTICAST_DELEGATE_TwoParams(FOnAddReplicationActivity, int64 ActivityId, bool bSuccess);
+		FOnAddReplicationActivity& OnAddReplicationActivity() { return OnAddReplicationActivityDelegate; }
+
+		//~ Begin IReplicationWorkspace Interface
+		virtual TOptional<int64> ProduceClientLeaveReplicationActivity(const FGuid& EndpointId, const FConcertSyncReplicationPayload_LeaveReplication& EventData) override;
+		virtual bool GetLastLeaveReplicationActivityByClient(const FConcertSessionClientInfo& InClientInfo, FConcertSyncReplicationPayload_LeaveReplication& OutLeaveReplication) const override;
+		virtual bool GetLeaveReplicationActivityById(const int64 ActivityId, FConcertSyncReplicationPayload_LeaveReplication& OutLeaveReplication) const override;
+		//~ End IReplicationWorkspace Interface
+
+	private:
+
+		/** Needed to get replication data. */
+		FConcertSyncSessionDatabase& Database;
+		/** Needed by GetLastLeaveReplicationActivityByClient to get the most appropriate activity data. */
+		const FFindSessionClient FindSessionClientDelegate;
+		/** Needed by ProduceClientLeaveReplicationActivity to correctly build the activity data. */
+		const FShouldIgnoreClientActivityOnRestore ShouldIgnoreClientActivityOnRestoreDelegate;
+		
+		FOnAddReplicationActivity OnAddReplicationActivityDelegate;
+	};
+}
+
diff --git a/Engine/Plugins/Developer/Concert/ConcertSync/ConcertSyncServer/Source/ConcertSyncServer/Public/Replication/IReplicationWorkspace.h b/Engine/Plugins/Developer/Concert/ConcertSync/ConcertSyncServer/Source/ConcertSyncServer/Public/Replication/IReplicationWorkspace.h
--- a/Engine/Plugins/Developer/Concert/ConcertSync/ConcertSyncServer/Source/ConcertSyncServer/Public/Replication/IReplicationWorkspace.h
+++ b/Engine/Plugins/Developer/Concert/ConcertSync/ConcertSyncServer/Source/ConcertSyncServer/Public/Replication/IReplicationWorkspace.h
@@ -3,8 +3,9 @@
 #pragma once
 
 #include "HAL/Platform.h"
+#include "Misc/Optional.h"
 
-struct FConcertClientInfo;
+struct FConcertSessionClientInfo;
 struct FConcertSyncReplicationPayload_LeaveReplication;
 struct FGuid;
 
@@ -18,8 +19,11 @@ namespace UE::ConcertSyncServer::Replication
 	{
 	public:
 		
-		/** Creates a replication activity for the client leaving replication if the session has the EConcertSyncSessionFlags::ShouldEnableReplicationActivities flag. */
-		virtual void ProduceClientLeaveReplicationActivity(const FGuid& EndpointId, const FConcertSyncReplicationPayload_LeaveReplication& EventData) = 0;
+		/**
+		 * Creates a replication activity for the client leaving replication if the session has the EConcertSyncSessionFlags::ShouldEnableReplicationActivities flag.
+		 * @return The identifier of the produced activity. Unset if activity insertion failed.
+		 */
+		virtual TOptional<int64> ProduceClientLeaveReplicationActivity(const FGuid& EndpointId, const FConcertSyncReplicationPayload_LeaveReplication& EventData) = 0;
 
 		/**
 		 * Gets the last replication leave activity associated for a given client.
@@ -31,7 +35,7 @@ namespace UE::ConcertSyncServer::Replication
 		 * @param OutLeaveReplication The activity, if present
 		 * @return Whether OutLeaveReplication contains a valid result.
 		 */
-		virtual bool GetLastLeaveReplicationActivityByClient(const FConcertClientInfo& InClientInfo, FConcertSyncReplicationPayload_LeaveReplication& OutLeaveReplication) const = 0;
+		virtual bool GetLastLeaveReplicationActivityByClient(const FConcertSessionClientInfo& InClientInfo, FConcertSyncReplicationPayload_LeaveReplication& OutLeaveReplication) const = 0;
 		/** Gets the replication leave activity with ActivityId. */
 		virtual bool GetLeaveReplicationActivityById(const int64 ActivityId, FConcertSyncReplicationPayload_LeaveReplication& OutLeaveReplication) const = 0;
 
diff --git a/Engine/Plugins/Developer/Concert/ConcertSync/ConcertSyncTest/Source/ConcertSyncTest/Private/Replication/ReplicationTestInterface.h b/Engine/Plugins/Developer/Concert/ConcertSync/ConcertSyncTest/Source/ConcertSyncTest/Private/Replication/ReplicationTestInterface.h
--- a/Engine/Plugins/Developer/Concert/ConcertSync/ConcertSyncTest/Source/ConcertSyncTest/Private/Replication/ReplicationTestInterface.h
+++ b/Engine/Plugins/Developer/Concert/ConcertSync/ConcertSyncTest/Source/ConcertSyncTest/Private/Replication/ReplicationTestInterface.h
@@ -3,12 +3,18 @@
 #pragma once
 
 #include "ConcertSyncSessionFlags.h"
+
+#include "Misc/Optional.h"
+#include "Templates/Function.h"
 #include "Templates/SharedPointer.h"
 
 class IConcertClientSession;
 class IConcertServerSession;
 class IConcertClientReplicationBridge;
 class IConcertClientReplicationManager;
+class FConcertSyncSessionDatabase;
+struct FConcertSessionClientInfo;
+struct FGuid;
 
 namespace UE::ConcertSyncServer::Replication
 {
@@ -34,4 +40,10 @@ namespace UE::ConcertSyncServer::TestInterface
 		Replication::IReplicationWorkspace& InWorkspace UE_LIFETIMEBOUND,
 		EConcertSyncSessionFlags InSessionFlags = EConcertSyncSessionFlags::Default_MultiUserSession
 		);
+
+	extern CONCERTSYNCSERVER_API TSharedRef<Replication::IReplicationWorkspace> CreateReplicationWorkspace(
+		FConcertSyncSessionDatabase& Database,
+		TFunction<TOptional<FConcertSessionClientInfo>(const FGuid& EndpointId)> FindSessionClient,
+		TFunction<bool(const FGuid& ClientId)> ShouldIgnoreClientActivityOnRestore
+		);
 }
\ No newline at end of file
diff --git a/Engine/Plugins/Developer/Concert/ConcertSync/ConcertSyncTest/Source/ConcertSyncTest/Private/Replication/ReplicationWorkspace.spec.cpp b/Engine/Plugins/Developer/Concert/ConcertSync/ConcertSyncTest/Source/ConcertSyncTest/Private/Replication/ReplicationWorkspace.spec.cpp
new file mode 100644
--- /dev/null
+++ b/Engine/Plugins/Developer/Concert/ConcertSync/ConcertSyncTest/Source/ConcertSyncTest/Private/Replication/ReplicationWorkspace.spec.cpp
@@ -0,0 +1,163 @@
+﻿// Copyright Epic Games, Inc. All Rights Reserved.
+
+#include "ReplicationTestInterface.h"
+#include "Misc/AutomationTest.h"
+#include "Replication/SyncControlState.h"
+#include "Replication/Messages/ReplicationActivity.h"
+#include "Replication/Messages/SyncControl.h"
+#include "Util/ConcertClientReplicationBridgeMock.h"
+#include "Util/ScopedSessionDatabase.h"
+#include "Util/Spec/ObjectTestReplicator.h"
+#include "Util/Spec/ReplicationClient.h"
+#include "Util/Spec/ReplicationServer.h"
+
+namespace UE::ConcertSyncTests::Replication::UI
+{
+	BEGIN_DEFINE_SPEC(FReplicationWorkspaceSpec, "Editor.Concert.Replication.ReplicationWorkspace", EAutomationTestFlags::EditorContext | EAutomationTestFlags::EngineFilter)
+		TUniquePtr<FScopedSessionDatabase> SessionDatabase;
+		TSharedPtr<ConcertSyncServer::Replication::IReplicationWorkspace> ReplicationWorkspace;
+	
+		TUniquePtr<FObjectTestReplicator> Object;
+		TUniquePtr<FReplicationServer> Server;
+		FReplicationClient* Client = nullptr;
+
+		FGuid StreamId = FGuid::NewGuid();
+		FGuid SecondaryStreamId = FGuid::NewGuid();
+
+		FReplicationClient& ConnectClient(const FString& DisplayName, const FString& DeviceName)
+		{
+			FReplicationClient& CreatedClient = Server->ConnectClient({ .DeviceName = DeviceName, .DisplayName = DisplayName });
+			SessionDatabase->SetEndpoint(CreatedClient.GetEndpointId(), { CreatedClient.GetClientInfo() });
+			return CreatedClient;
+		}
+
+		FConcertSyncReplicationPayload_LeaveReplication MakeExpectedData() const { return { .Streams = { Object->CreateStream(StreamId) } }; }
+		FConcertSessionClientInfo GetClientSessionInfo() const { return { Client->GetEndpointId(), Client->GetClientInfo() }; }
+		static FString MainDisplayName() { return TEXT("PrimaryClientName"); }
+		static FString MainDeviceName() { return TEXT("PrimaryMachine"); }
+		static FString SecondaryDeviceName() { return TEXT("SecondaryMachine"); }
+	END_DEFINE_SPEC(FReplicationWorkspaceSpec);
+
+	/** This tests that clients FReplicationWorkspace is correctly implemented. */
+	void FReplicationWorkspaceSpec::Define()
+	{
+		BeforeEach([this]
+		{
+			SessionDatabase = MakeUnique<FScopedSessionDatabase>(*this);
+			const auto FindSessionClient = [this](const FGuid& EndpointId) -> TOptional<FConcertSessionClientInfo>
+			{
+				FConcertSessionClientInfo Info;
+				const bool bFound = Server->GetServerSessionMock()->FindSessionClient(EndpointId, Info);
+				return bFound ? Info : TOptional<FConcertSessionClientInfo>{};
+			};
+			const auto ShouldIgnoreClientActivityOnRestore = [](const FGuid& EndpointId){ return false; };
+			ReplicationWorkspace = ConcertSyncServer::TestInterface::CreateReplicationWorkspace(*SessionDatabase, FindSessionClient, ShouldIgnoreClientActivityOnRestore);
+			
+			Object = MakeUnique<FObjectTestReplicator>();
+			Server = MakeUnique<FReplicationServer>(*this);
+			Client = &ConnectClient(MainDisplayName(), MainDeviceName());
+			Client->JoinReplication();
+		});
+		AfterEach([this]
+		{
+			ReplicationWorkspace.Reset();
+			SessionDatabase.Reset();
+			Server.Reset();
+		});
+		
+		It("ProduceClientLeaveReplicationActivity", [this]
+		{
+			const FConcertSyncReplicationPayload_LeaveReplication ExpectedData = MakeExpectedData();
+			const FGuid ClientId = Client->GetEndpointId();
+			const TOptional<int64> ActivityId = ReplicationWorkspace->ProduceClientLeaveReplicationActivity(ClientId, ExpectedData);
+			if (!ActivityId)
+			{
+				AddError(TEXT("Failed to create activity"));
+				return;
+			}
+
+			FConcertSyncReplicationActivity Activity;
+			if (!SessionDatabase->GetReplicationActivity(*ActivityId, Activity))
+			{
+				AddError(TEXT("Activity not saved"));
+				return;
+			}
+
+			FConcertSyncReplicationPayload_LeaveReplication ActualData;
+			const bool bGotData = Activity.EventData.GetPayload(ActualData);
+			TestTrue(TEXT("Got Data"), bGotData);
+			TestEqual(TEXT("Payload"), ActualData, ExpectedData);
+		});
+
+		Describe("GetLastLeaveReplicationActivityByClient", [this]
+		{
+			BeforeEach([this]()
+			{
+				const FGuid ClientId = Client->GetEndpointId();
+				ReplicationWorkspace->ProduceClientLeaveReplicationActivity(ClientId, MakeExpectedData());
+			});
+			
+			It("Simple: Single endpoint", [this]
+			{
+				FConcertSyncReplicationPayload_LeaveReplication ActualData;
+				const bool bGotData = ReplicationWorkspace->GetLastLeaveReplicationActivityByClient(GetClientSessionInfo(), ActualData);
+				
+				TestTrue(TEXT("Got Data"), bGotData);
+				TestEqual(TEXT("Payload"), ActualData, MakeExpectedData());
+			});
+			It("When single client has changed device, use the latest data", [this]
+			{
+				const FReplicationClient& OtherClient = ConnectClient(MainDisplayName(), MainDeviceName());
+				
+				FConcertSyncReplicationPayload_LeaveReplication ActualData;
+				const bool bGotData = ReplicationWorkspace->GetLastLeaveReplicationActivityByClient({ OtherClient.GetEndpointId(), OtherClient.GetClientInfo() }, ActualData);
+				TestTrue(TEXT("Got Data"), bGotData);
+				TestEqual(TEXT("Payload"), ActualData, MakeExpectedData());
+			});
+			
+			It("When there are 2 clients with same display name but different device names, use the data associated with the device name", [this]
+			{
+				const FReplicationClient& OtherClient = ConnectClient(MainDisplayName(), SecondaryDeviceName());
+				const FConcertSyncReplicationPayload_LeaveReplication OtherClientData { .Streams = { Object->CreateStream(SecondaryStreamId) } };
+				const TOptional<int64> OtherClientActivityId = ReplicationWorkspace->ProduceClientLeaveReplicationActivity(OtherClient.GetEndpointId(), OtherClientData);
+				
+				FConcertSyncReplicationPayload_LeaveReplication ActualData;
+				const bool bGotData = ReplicationWorkspace->GetLastLeaveReplicationActivityByClient(GetClientSessionInfo(), ActualData);
+				TestTrue(TEXT("Got Data"), bGotData);
+				TestEqual(TEXT("Payload"), ActualData, MakeExpectedData());
+				// This is already handled by the Payload case but we'll make sure
+				TestTrue(TEXT("StreamId"), ActualData.Streams.Num() == 1 && ActualData.Streams[0].BaseDescription.Identifier == StreamId);
+			});
+			It ("When there are 2 clients with the same display and device name, use the latest data", [this]
+			{
+				const FReplicationClient& OtherClient = ConnectClient(MainDisplayName(), MainDeviceName());
+				const FConcertSyncReplicationPayload_LeaveReplication OtherClientData { .Streams = { Object->CreateStream(SecondaryStreamId) } };
+				const TOptional<int64> OtherClientActivityId = ReplicationWorkspace->ProduceClientLeaveReplicationActivity(OtherClient.GetEndpointId(), OtherClientData);
+				
+				FConcertSyncReplicationPayload_LeaveReplication ActualData;
+				const bool bGotData = ReplicationWorkspace->GetLastLeaveReplicationActivityByClient(GetClientSessionInfo(), ActualData);
+				TestTrue(TEXT("Got Data"), bGotData);
+				TestEqual(TEXT("Payload"), ActualData, OtherClientData);
+				// This is already handled by the Payload case but we'll make sure
+				TestTrue(TEXT("StreamId"), ActualData.Streams.Num() == 1 && ActualData.Streams[0].BaseDescription.Identifier == SecondaryStreamId);
+			});
+		});
+
+		It("GetLeaveReplicationActivityById", [this]
+		{
+			const FConcertSyncReplicationPayload_LeaveReplication ExpectedData = MakeExpectedData();
+			const FGuid ClientId = Client->GetEndpointId();
+			const TOptional<int64> ActivityId = ReplicationWorkspace->ProduceClientLeaveReplicationActivity(ClientId, ExpectedData);
+			if (!ActivityId)
+			{
+				AddError(TEXT("Failed to create activity"));
+				return;
+			}
+
+			FConcertSyncReplicationPayload_LeaveReplication ActualData;
+			const bool bGotData = ReplicationWorkspace->GetLeaveReplicationActivityById(*ActivityId, ActualData);
+			TestTrue(TEXT("Got Data"), bGotData);
+			TestEqual(TEXT("Payload"), ActualData, ExpectedData);
+		});
+	}
+}
diff --git a/Engine/Plugins/Developer/Concert/ConcertSync/ConcertSyncTest/Source/ConcertSyncTest/Private/Replication/Util/Mocks/ReplicationWorkspaceCallInterceptorMock.h b/Engine/Plugins/Developer/Concert/ConcertSync/ConcertSyncTest/Source/ConcertSyncTest/Private/Replication/Util/Mocks/ReplicationWorkspaceCallInterceptorMock.h
--- a/Engine/Plugins/Developer/Concert/ConcertSync/ConcertSyncTest/Source/ConcertSyncTest/Private/Replication/Util/Mocks/ReplicationWorkspaceCallInterceptorMock.h
+++ b/Engine/Plugins/Developer/Concert/ConcertSync/ConcertSyncTest/Source/ConcertSyncTest/Private/Replication/Util/Mocks/ReplicationWorkspaceCallInterceptorMock.h
@@ -4,7 +4,6 @@
 
 #include "Replication/IReplicationWorkspace.h"
 #include "Replication/Messages/ReplicationActivity.h"
-#include "Replication/Messages/RestoreContent.h"
 
 #include "Misc/Optional.h"
 #include "Misc/Guid.h"
@@ -19,21 +18,24 @@ namespace UE::ConcertSyncTests::Replication
 		/** Arguments of last ProduceClientLeaveReplicationActivity call. */
 		TOptional<TTuple<FGuid, FConcertSyncReplicationPayload_LeaveReplication>> LastCall_ProduceClientLeaveReplicationActivity;
 		/** Arguments of last GetLastLeaveReplicationActivityByClient call. */
-		mutable TOptional<TTuple<FConcertClientInfo>> LastCall_GetLastLeaveReplicationActivityByClient;
+		mutable TOptional<TTuple<FConcertSessionClientInfo>> LastCall_GetLastLeaveReplicationActivityByClient;
 		/** Arguments of last GetLastLeaveReplicationActivityByClient call. */
 		mutable TOptional<TTuple<int64>> LastCall_GetLeaveReplicationActivityById;
 
+		/** The result to return in ProduceClientLeaveReplicationActivity. */
+		TOptional<int64> ReturnResult_ProduceClientLeaveReplicationActivity = 0;
 		/** The result to return in GetLastLeaveReplicationActivityByClient. */
 		TOptional<FConcertSyncReplicationPayload_LeaveReplication> ReturnResult_GetLastLeaveReplicationActivityByClient;
 		/** The result to return in GetLastLeaveReplicationActivityByClient. */
 		TOptional<FConcertSyncReplicationPayload_LeaveReplication> ReturnResult_GetLeaveReplicationActivityById;
 		
 		//~ Begin IReplicationWorkspace Interface
-		virtual void ProduceClientLeaveReplicationActivity(const FGuid& EndpointId, const FConcertSyncReplicationPayload_LeaveReplication& EventData) override
+		virtual TOptional<int64> ProduceClientLeaveReplicationActivity(const FGuid& EndpointId, const FConcertSyncReplicationPayload_LeaveReplication& EventData) override
 		{
 			LastCall_ProduceClientLeaveReplicationActivity = MakeTuple(EndpointId, EventData);
+			return ReturnResult_ProduceClientLeaveReplicationActivity;
 		}
-		virtual bool GetLastLeaveReplicationActivityByClient(const FConcertClientInfo& InClientInfo, FConcertSyncReplicationPayload_LeaveReplication& OutLeaveReplication) const override
+		virtual bool GetLastLeaveReplicationActivityByClient(const FConcertSessionClientInfo& InClientInfo, FConcertSyncReplicationPayload_LeaveReplication& OutLeaveReplication) const override
 		{
 			LastCall_GetLastLeaveReplicationActivityByClient = MakeTuple(InClientInfo);
 			if (ReturnResult_GetLastLeaveReplicationActivityByClient)
diff --git a/Engine/Plugins/Developer/Concert/ConcertSync/ConcertSyncTest/Source/ConcertSyncTest/Private/Replication/Util/Mocks/ReplicationWorkspaceEmptyMock.h b/Engine/Plugins/Developer/Concert/ConcertSync/ConcertSyncTest/Source/ConcertSyncTest/Private/Replication/Util/Mocks/ReplicationWorkspaceEmptyMock.h
--- a/Engine/Plugins/Developer/Concert/ConcertSync/ConcertSyncTest/Source/ConcertSyncTest/Private/Replication/Util/Mocks/ReplicationWorkspaceEmptyMock.h
+++ b/Engine/Plugins/Developer/Concert/ConcertSync/ConcertSyncTest/Source/ConcertSyncTest/Private/Replication/Util/Mocks/ReplicationWorkspaceEmptyMock.h
@@ -11,8 +11,8 @@ namespace UE::ConcertSyncTests::Replication
 	public:
 		
 		//~ Begin IReplicationWorkspace Interface
-		virtual void ProduceClientLeaveReplicationActivity(const FGuid&, const FConcertSyncReplicationPayload_LeaveReplication&) override {}
-		virtual bool GetLastLeaveReplicationActivityByClient(const FConcertClientInfo& InClientInfo, FConcertSyncReplicationPayload_LeaveReplication& OutLeaveReplication) const override { return false; }
+		virtual TOptional<int64> ProduceClientLeaveReplicationActivity(const FGuid&, const FConcertSyncReplicationPayload_LeaveReplication&) override { return {}; }
+		virtual bool GetLastLeaveReplicationActivityByClient(const FConcertSessionClientInfo& InClientInfo, FConcertSyncReplicationPayload_LeaveReplication& OutLeaveReplication) const override { return false; }
 		virtual bool GetLeaveReplicationActivityById(const int64 ActivityId, FConcertSyncReplicationPayload_LeaveReplication& OutLeaveReplication) const override { return false; }
 		//~ End IReplicationWorkspace Interface
 	};
