Restore mute state upon joining session

diff --git a/Engine/Plugins/Developer/Concert/ConcertApp/MultiUserClient/Source/MultiUserClient/Private/Replication/MultiUserReplicationManager.cpp b/Engine/Plugins/Developer/Concert/ConcertApp/MultiUserClient/Source/MultiUserClient/Private/Replication/MultiUserReplicationManager.cpp
--- a/Engine/Plugins/Developer/Concert/ConcertApp/MultiUserClient/Source/MultiUserClient/Private/Replication/MultiUserReplicationManager.cpp
+++ b/Engine/Plugins/Developer/Concert/ConcertApp/MultiUserClient/Source/MultiUserClient/Private/Replication/MultiUserReplicationManager.cpp
@@ -113,7 +113,7 @@ namespace UE::MultiUserClient
 	{
 		Client->GetReplicationManager()->RestoreContent(
 			{
-				.Flags = EConcertReplicationRestoreContentFlags::StreamsAndAuthority | EConcertReplicationRestoreContentFlags::ValidateUniqueClient
+				.Flags = EConcertReplicationRestoreContentFlags::All | EConcertReplicationRestoreContentFlags::ValidateUniqueClient
 			}
 		)
 		.Next([ClientInfo = Client->GetConcertClient()->GetClientInfo()](FConcertReplication_RestoreContent_Response&& Response)
diff --git a/Engine/Plugins/Developer/Concert/ConcertSync/ConcertSyncCore/Source/ConcertSyncCore/Public/Replication/Messages/ReplicationActivity.h b/Engine/Plugins/Developer/Concert/ConcertSync/ConcertSyncCore/Source/ConcertSyncCore/Public/Replication/Messages/ReplicationActivity.h
--- a/Engine/Plugins/Developer/Concert/ConcertSync/ConcertSyncCore/Source/ConcertSyncCore/Public/Replication/Messages/ReplicationActivity.h
+++ b/Engine/Plugins/Developer/Concert/ConcertSync/ConcertSyncCore/Source/ConcertSyncCore/Public/Replication/Messages/ReplicationActivity.h
@@ -168,7 +168,7 @@ struct FConcertSyncReplicationActivity : public FConcertSyncActivity
 		EventType = EConcertSyncActivityEventType::Replication;
 	}
 
-	/** The transaction event data associated with this activity */
+	/** The replication event data associated with this activity */
 	UPROPERTY()
 	FConcertSyncReplicationEvent EventData;
 };
diff --git a/Engine/Plugins/Developer/Concert/ConcertSync/ConcertSyncCore/Source/ConcertSyncCore/Public/Replication/Messages/RestoreContent.h b/Engine/Plugins/Developer/Concert/ConcertSync/ConcertSyncCore/Source/ConcertSyncCore/Public/Replication/Messages/RestoreContent.h
--- a/Engine/Plugins/Developer/Concert/ConcertSync/ConcertSyncCore/Source/ConcertSyncCore/Public/Replication/Messages/RestoreContent.h
+++ b/Engine/Plugins/Developer/Concert/ConcertSync/ConcertSyncCore/Source/ConcertSyncCore/Public/Replication/Messages/RestoreContent.h
@@ -28,7 +28,11 @@ enum class EConcertReplicationRestoreContentFlags : uint8
 	/** The response should include the new stream and authority state. */
 	SendNewState = 1 << 3,
 
-	StreamsAndAuthority = RestoreAuthority
+	/** Also restore the mute states of the objects. This is done by going through the activity history and replaying the mute & unmute actions. */
+	RestoreMute = 1 << 4,
+
+	StreamsAndAuthority = RestoreAuthority,
+	All = StreamsAndAuthority | RestoreMute
 };
 ENUM_CLASS_FLAGS(EConcertReplicationRestoreContentFlags);
 
diff --git a/Engine/Plugins/Developer/Concert/ConcertSync/ConcertSyncServer/Source/ConcertSyncServer/Private/Replication/ConcertServerReplicationManager.h b/Engine/Plugins/Developer/Concert/ConcertSync/ConcertSyncServer/Source/ConcertSyncServer/Private/Replication/ConcertServerReplicationManager.h
--- a/Engine/Plugins/Developer/Concert/ConcertSync/ConcertSyncServer/Source/ConcertSyncServer/Private/Replication/ConcertServerReplicationManager.h
+++ b/Engine/Plugins/Developer/Concert/ConcertSync/ConcertSyncServer/Source/ConcertSyncServer/Private/Replication/ConcertServerReplicationManager.h
@@ -137,6 +137,10 @@ namespace UE::ConcertSyncServer::Replication
 		/** Restores the authority in DataToApply to Client. */
 		void RestoreAuthority(const FConcertSyncReplicationPayload_LeaveReplication& DataToApply, const FConcertReplicationClient& Client, FConcertReplication_ChangeSyncControl& OutChangedSyncControl);
 		
+		/** Restores the mute state of Client. */
+		void RestoreMuteState(const FConcertReplicationClient& Client, FConcertReplication_ChangeSyncControl& OutChangedSyncControl);
+		void ApplyRestoringMuteRequest(const FConcertReplicationClient& Client, const FConcertReplication_ChangeMuteState_Request& AggregatedRequest, FConcertReplication_ChangeSyncControl& OutChangedSyncControl);
+		
 		// Leaving
 		void HandleLeaveReplicationSessionRequest(const FConcertSessionContext& ConcertSessionContext, const FConcertReplication_LeaveEvent& EventData);
 		void OnConnectionChanged(IConcertServerSession& ConcertServerSession, EConcertClientStatus ConcertClientStatus, const FConcertSessionClientInfo& ConcertSessionClientInfo);
diff --git a/Engine/Plugins/Developer/Concert/ConcertSync/ConcertSyncServer/Source/ConcertSyncServer/Private/Replication/ConcertServerReplicationManager_RestoreContent.cpp b/Engine/Plugins/Developer/Concert/ConcertSync/ConcertSyncServer/Source/ConcertSyncServer/Private/Replication/ConcertServerReplicationManager_RestoreContent.cpp
--- a/Engine/Plugins/Developer/Concert/ConcertSync/ConcertSyncServer/Source/ConcertSyncServer/Private/Replication/ConcertServerReplicationManager_RestoreContent.cpp
+++ b/Engine/Plugins/Developer/Concert/ConcertSync/ConcertSyncServer/Source/ConcertSyncServer/Private/Replication/ConcertServerReplicationManager_RestoreContent.cpp
@@ -4,6 +4,7 @@
 #include "ConcertServerReplicationManager.h"
 #include "ConcertServer/Private/ConcertServerSession.h"
 #include "Replication/IReplicationWorkspace.h"
+#include "Replication/MuteUtils.h"
 #include "Replication/Messages/ChangeAuthority.h"
 #include "Replication/Messages/ChangeStream.h"
 #include "Replication/Messages/ReplicationActivity.h"
@@ -271,7 +272,7 @@ namespace UE::ConcertSyncServer::Replication
 	{
 		FConcertReplicationClient& Client = *Clients[RequestingEndpointId];
 		
-		// After this first step, SyncControl will hold only the objects that were removed thus far.
+		// After this first step, SyncControl hold the objects that the client had sync control over before but that are removed by this request.
 		RestoreStreamContent(Request, DataToApply, Client, ChangedSyncControl);
 		
 		if (EnumHasAnyFlags(Request.Flags, EConcertReplicationRestoreContentFlags::RestoreAuthority))
@@ -279,6 +280,11 @@ namespace UE::ConcertSyncServer::Replication
 			// This adds the objects the client receives sync control for
 			RestoreAuthority(DataToApply, Client, ChangedSyncControl);
 		}
+
+		if (EnumHasAnyFlags(Request.Flags, EConcertReplicationRestoreContentFlags::RestoreMute))
+		{
+			RestoreMuteState(Client, ChangedSyncControl);
+		}
 	}
 	
 	void FConcertServerReplicationManager::RestoreStreamContent(
@@ -348,4 +354,99 @@ namespace UE::ConcertSyncServer::Replication
 		AuthorityManager.ApplyChangeAuthorityRequest(Client.GetClientEndpointId(), Request, RejectedObjects, OutChangedSyncControl);
 		check(RejectedObjects.IsEmpty());
 	}
-}
\ No newline at end of file
+
+	void FConcertServerReplicationManager::RestoreMuteState(const FConcertReplicationClient& Client, FConcertReplication_ChangeSyncControl& OutChangedSyncControl)
+	{
+		class FMuteStateGroundTruth : public ConcertSyncCore::Replication::MuteUtils::IMuteStateGroundTruth
+		{
+			const FMuteManager& MuteManager;
+			const ConcertSyncCore::FReplicatedObjectHierarchyCache& ServerObjectCache;
+		public:
+			
+			FMuteStateGroundTruth(
+				const FMuteManager& MuteManager UE_LIFETIMEBOUND,
+				const ConcertSyncCore::FReplicatedObjectHierarchyCache& ServerObjectCache UE_LIFETIMEBOUND
+				)
+				: MuteManager(MuteManager)
+				, ServerObjectCache(ServerObjectCache)
+			{}
+
+			virtual ConcertSyncCore::Replication::MuteUtils::EMuteState GetMuteState(const FSoftObjectPath& Object) const override
+			{
+				using namespace ConcertSyncCore::Replication::MuteUtils;
+				
+				const TOptional<FMuteManager::EMuteState> MuteState = MuteManager.GetMuteState(Object);
+				if (!MuteState)
+				{
+					return EMuteState::None;
+				}
+
+				switch (*MuteState)
+				{
+				case FMuteManager::EMuteState::ExplicitlyMuted: return EMuteState::ExplicitlyMuted;
+				case FMuteManager::EMuteState::ExplicitlyUnmuted: return EMuteState::ExplicitlyUnmuted;
+				case FMuteManager::EMuteState::ImplicitlyMuted: return EMuteState::ImplicitlyMuted;
+				case FMuteManager::EMuteState::ImplicitlyUnmuted: return EMuteState::ImplicitlyUnmuted;
+				default: checkNoEntry(); return EMuteState::None;
+				}
+			}
+			
+			virtual TOptional<FConcertReplication_ObjectMuteSetting> GetExplicitSetting(const FSoftObjectPath& Object) const override
+			{
+				return MuteManager.GetExplicitMuteSetting(Object);
+			}
+			
+			virtual bool IsObjectKnown(const FSoftObjectPath& Object) const override
+			{
+				return ServerObjectCache.IsInHierarchy(Object).IsSet();
+			}
+		} GroundTruth(MuteManager, ServerObjectCache);
+
+		// We'll effectively replay all mute actions that have occured so far by combining them into Request. 
+		FConcertReplication_ChangeMuteState_Request AggregatedRequest;
+		ServerWorkspace.EnumerateMuteActivities([&GroundTruth, &AggregatedRequest](const FConcertSyncReplicationActivity& Activity)
+		{
+			FConcertSyncReplicationPayload_Mute MuteData;
+			if (!ensure(Activity.EventData.ActivityType == EConcertSyncReplicationActivityType::Mute)
+				|| !Activity.EventData.GetPayload(MuteData))
+			{
+				return EBreakBehavior::Continue;
+			}
+
+			// This skips changes that are already in effect or that are invalid to do (e.g. due to unknown object)
+			CombineMuteRequests(AggregatedRequest, MuteData.Request, GroundTruth);
+			return EBreakBehavior::Continue;
+		});
+
+		if (!AggregatedRequest.IsEmpty())
+		{
+			ApplyRestoringMuteRequest(Client, AggregatedRequest, OutChangedSyncControl);
+		}
+	}
+	void FConcertServerReplicationManager::ApplyRestoringMuteRequest(
+		const FConcertReplicationClient& Client,
+		const FConcertReplication_ChangeMuteState_Request& AggregatedRequest,
+		FConcertReplication_ChangeSyncControl& OutChangedSyncControl
+		)
+	{
+		const FGuid& ClientId = Client.GetClientEndpointId();
+		const FConcertReplication_ChangeSyncControl GainedSyncControl = MuteManager.ApplyManualRequest(ClientId, AggregatedRequest);
+
+		// ApplyManualRequest may have removed sync control but GainedSyncControl does not contain the removed objects...
+		for (TPair<FConcertObjectInStreamID, bool>& ResultControl : OutChangedSyncControl.NewControlStates)
+		{
+			const FConcertReplicatedObjectId ReplicatedObjectId { { ResultControl.Key }, ClientId };
+			// ... so simply go through everything and check whether it lost control
+			if (!SyncControlManager.HasSyncControl(ReplicatedObjectId))
+			{
+				ResultControl.Value = false;
+			}
+		}
+
+		// And all the things that have now gained the control will be put into OutChangedSyncControl afterwards.
+		for (const TPair<FConcertObjectInStreamID, bool>& GainedControl : GainedSyncControl.NewControlStates)
+		{
+			OutChangedSyncControl.NewControlStates.Add(GainedControl.Key, GainedControl.Value);
+		}
+	}
+}
diff --git a/Engine/Plugins/Developer/Concert/ConcertSync/ConcertSyncServer/Source/ConcertSyncServer/Private/Replication/MuteManager.cpp b/Engine/Plugins/Developer/Concert/ConcertSync/ConcertSyncServer/Source/ConcertSyncServer/Private/Replication/MuteManager.cpp
--- a/Engine/Plugins/Developer/Concert/ConcertSync/ConcertSyncServer/Source/ConcertSyncServer/Private/Replication/MuteManager.cpp
+++ b/Engine/Plugins/Developer/Concert/ConcertSync/ConcertSyncServer/Source/ConcertSyncServer/Private/Replication/MuteManager.cpp
@@ -47,6 +47,28 @@ namespace UE::ConcertSyncServer::Replication
 			&& (ObjectMuteState->State == EMuteState::ExplicitlyMuted || ObjectMuteState->State == EMuteState::ImplicitlyMuted);
 	}
 
+	TOptional<FMuteManager::EMuteState> FMuteManager::GetMuteState(const FSoftObjectPath& Object) const
+	{
+		const FMuteData* ObjectMuteState = MuteState.Find(Object);
+		return ObjectMuteState
+			? ObjectMuteState->State
+			: TOptional<EMuteState>{};
+	}
+
+	TOptional<FConcertReplication_ObjectMuteSetting> FMuteManager::GetExplicitMuteSetting(const FSoftObjectPath& Object) const
+	{
+		const FMuteData* ObjectMuteState = MuteState.Find(Object);
+		return ObjectMuteState
+			? ObjectMuteState->MuteSetting
+			: TOptional<FConcertReplication_ObjectMuteSetting>{};
+	}
+
+	FConcertReplication_ChangeSyncControl FMuteManager::ApplyManualRequest(const FGuid& EndpointId, const FConcertReplication_ChangeMuteState_Request& Request)
+	{
+		ApplyRequest(Request);
+		return OnGenerateSyncControlForMuteChangeDelegate.Execute(EndpointId);
+	}
+
 	void FMuteManager::PostApplyStreamChange(const FGuid& ClientId, TConstArrayView<FConcertObjectInStreamID> AddedObjects, TConstArrayView<FConcertObjectInStreamID> RemovedObjects)
 	{
 		// If Request adds new subobjects that should be muted implicitly by existing rules, update MuteState.
@@ -213,8 +235,8 @@ namespace UE::ConcertSyncServer::Replication
 		Response.ErrorCode = EConcertReplicationMuteErrorCode::Accepted;
 		ApplyRequest(Request);
 		
-		// Sync control will now 1. generate a response we can embed in the response, and 2. send a network event to all other clients.
-		// The embedded response will contain only those objects which changed sync control
+		// Sync control will now 1. generate a sync control we can embed in the response, and 2. send a network event to all other clients.
+		// The embedded sync control will contain only those objects which changed sync control
 		// Fyi, the new sync control may contain more objects than were specified in the request since parent objects can also cause their subobjects' mute state to change.
 		check(OnGenerateSyncControlForMuteChangeDelegate.IsBound());
 		Response.SyncControl = OnGenerateSyncControlForMuteChangeDelegate.Execute(Context.SourceEndpointId);
diff --git a/Engine/Plugins/Developer/Concert/ConcertSync/ConcertSyncServer/Source/ConcertSyncServer/Private/Replication/MuteManager.h b/Engine/Plugins/Developer/Concert/ConcertSync/ConcertSyncServer/Source/ConcertSyncServer/Private/Replication/MuteManager.h
--- a/Engine/Plugins/Developer/Concert/ConcertSync/ConcertSyncServer/Source/ConcertSyncServer/Private/Replication/MuteManager.h
+++ b/Engine/Plugins/Developer/Concert/ConcertSync/ConcertSyncServer/Source/ConcertSyncServer/Private/Replication/MuteManager.h
@@ -6,6 +6,8 @@
 #include "ConcertSyncSessionFlags.h"
 #include "IConcertSessionHandler.h"
 #include "Replication/Messages/Muting.h"
+
+#include "Misc/Optional.h"
 #include "Templates/UnrealTemplate.h"
 
 #include <type_traits>
@@ -46,6 +48,19 @@ namespace UE::ConcertSyncServer::Replication
 		
 		/** Notifies manager about applied event. Generates mute activity. */
 		DECLARE_MULTICAST_DELEGATE_TwoParams(FOnMuteRequestApplied, const FGuid& ClientId, const FConcertReplication_ChangeMuteState_Request& Request);
+		
+		enum class EMuteState : uint8
+		{
+			/** Object is muted */
+			ExplicitlyMuted,
+			/** A parent object is ExplicitlyMuted and has EConcertReplicationMuteFlags::ObjectAndSubobjects set but is object and its children are supposed to be implicitly unmuted.  */
+			ExplicitlyUnmuted,
+			
+			/** Object is muted because one of its parent objects is ExplicitlyMuted and EConcertReplicationMuteFlags::ObjectAndSubobjects set. */
+			ImplicitlyMuted,
+			/** Object is unmuted because one of its parent objects is ExplicitlyUnmuted and has EConcertReplicationMuteFlags::ObjectAndSubobjects set. */
+			ImplicitlyUnmuted,
+		};
 
 		FMuteManager(
 			IConcertSession& InSession UE_LIFETIMEBOUND,
@@ -56,6 +71,17 @@ namespace UE::ConcertSyncServer::Replication
 
 		/** @return Whether the object is globally muted. */
 		bool IsMuted(const FSoftObjectPath& Object) const;
+
+		/** @return The mute state of Object. If it was not explicitly muted and no parent object affects this subobject, then the optional is unset. */
+		TOptional<EMuteState> GetMuteState(const FSoftObjectPath& Object) const;
+		/** @return The mute setting of Object, if it is explicitly set (i.e. not affected by a parent object). */
+		TOptional<FConcertReplication_ObjectMuteSetting> GetExplicitMuteSetting(const FSoftObjectPath& Object) const;
+
+		/**
+		 * Applies Request as if it was sent by EndpointId.
+		 * @return The sync control that the client identified by EndpointId has gained. Does NOT include the sync control they lost.
+		 */
+		FConcertReplication_ChangeSyncControl ApplyManualRequest(const FGuid& EndpointId, const FConcertReplication_ChangeMuteState_Request& Request);
 		
 		/** Called right after objects have been unregistered from ClientId's streams. */
 		void PostApplyStreamChange(const FGuid& ClientId, TConstArrayView<FConcertObjectInStreamID> AddedObjects, TConstArrayView<FConcertObjectInStreamID> RemovedObjects);
@@ -89,19 +115,6 @@ namespace UE::ConcertSyncServer::Replication
 		
 		/** Broadcasts after a mute request has been applied. */
 		FOnMuteRequestApplied OnMuteRequestAppliedDelegate;
-
-		enum class EMuteState : uint8
-		{
-			/** Object is muted */
-			ExplicitlyMuted,
-			/** A parent object is ExplicitlyMuted and has EConcertReplicationMuteFlags::ObjectAndSubobjects set but is object and its children are supposed to be implicitly unmuted.  */
-			ExplicitlyUnmuted,
-			
-			/** Object is muted because one of its parent objects is ExplicitlyMuted and EConcertReplicationMuteFlags::ObjectAndSubobjects set. */
-			ImplicitlyMuted,
-			/** Object is unmuted because one of its parent objects is ExplicitlyUnmuted and has EConcertReplicationMuteFlags::ObjectAndSubobjects set. */
-			ImplicitlyUnmuted,
-		};
 		
 		struct FMuteData
 		{
diff --git a/Engine/Plugins/Developer/Concert/ConcertSync/ConcertSyncServer/Source/ConcertSyncServer/Private/Replication/ReplicationWorkspace.cpp b/Engine/Plugins/Developer/Concert/ConcertSync/ConcertSyncServer/Source/ConcertSyncServer/Private/Replication/ReplicationWorkspace.cpp
--- a/Engine/Plugins/Developer/Concert/ConcertSync/ConcertSyncServer/Source/ConcertSyncServer/Private/Replication/ReplicationWorkspace.cpp
+++ b/Engine/Plugins/Developer/Concert/ConcertSync/ConcertSyncServer/Source/ConcertSyncServer/Private/Replication/ReplicationWorkspace.cpp
@@ -74,7 +74,16 @@ namespace UE::ConcertSyncServer
 		const bool bGotEventData = Database.GetReplicationEvent(ActivityId, Event);
 		return bGotEventData && Event.ActivityType == EConcertSyncReplicationActivityType::LeaveReplication && Event.GetPayload(OutLeaveReplication);
 	}
-	
+
+	void FReplicationWorkspace::EnumerateMuteActivities(TFunctionRef<EBreakBehavior(const FConcertSyncReplicationActivity& Activity)> Callback) const
+	{
+		Database.EnumerateReplicationActivities([&Callback](FConcertSyncReplicationActivity&& Activity)
+		{
+			return Activity.EventData.ActivityType != EConcertSyncReplicationActivityType::Mute
+				|| Callback(Activity) == EBreakBehavior::Continue;
+		});
+	}
+
 	template <typename TPayload>
 	TOptional<int64> FReplicationWorkspace::ProduceActivity(const FGuid& EndpointId, const TPayload& EventData)
 	{
diff --git a/Engine/Plugins/Developer/Concert/ConcertSync/ConcertSyncServer/Source/ConcertSyncServer/Private/Replication/ReplicationWorkspace.h b/Engine/Plugins/Developer/Concert/ConcertSync/ConcertSyncServer/Source/ConcertSyncServer/Private/Replication/ReplicationWorkspace.h
--- a/Engine/Plugins/Developer/Concert/ConcertSync/ConcertSyncServer/Source/ConcertSyncServer/Private/Replication/ReplicationWorkspace.h
+++ b/Engine/Plugins/Developer/Concert/ConcertSync/ConcertSyncServer/Source/ConcertSyncServer/Private/Replication/ReplicationWorkspace.h
@@ -7,6 +7,7 @@
 #include "Delegates/Delegate.h"
 #include "HAL/Platform.h"
 #include "Misc/Optional.h"
+#include "Templates/Function.h"
 
 class FConcertSyncSessionDatabase;
 
@@ -40,6 +41,7 @@ namespace UE::ConcertSyncServer
 		virtual TOptional<int64> ProduceClientMuteReplicationActivity(const FGuid& EndpointId, const FConcertSyncReplicationPayload_Mute& EventData) override;
 		virtual bool GetLastLeaveReplicationActivityByClient(const FConcertSessionClientInfo& InClientInfo, FConcertSyncReplicationPayload_LeaveReplication& OutLeaveReplication) const override;
 		virtual bool GetLeaveReplicationActivityById(const int64 ActivityId, FConcertSyncReplicationPayload_LeaveReplication& OutLeaveReplication) const override;
+		virtual void EnumerateMuteActivities(TFunctionRef<EBreakBehavior(const FConcertSyncReplicationActivity& Activity)> Callback) const override;
 		//~ End IReplicationWorkspace Interface
 
 	private:
diff --git a/Engine/Plugins/Developer/Concert/ConcertSync/ConcertSyncServer/Source/ConcertSyncServer/Public/Replication/IReplicationWorkspace.h b/Engine/Plugins/Developer/Concert/ConcertSync/ConcertSyncServer/Source/ConcertSyncServer/Public/Replication/IReplicationWorkspace.h
--- a/Engine/Plugins/Developer/Concert/ConcertSync/ConcertSyncServer/Source/ConcertSyncServer/Public/Replication/IReplicationWorkspace.h
+++ b/Engine/Plugins/Developer/Concert/ConcertSync/ConcertSyncServer/Source/ConcertSyncServer/Public/Replication/IReplicationWorkspace.h
@@ -4,7 +4,11 @@
 
 #include "HAL/Platform.h"
 #include "Misc/Optional.h"
+#include "Templates/FunctionFwd.h"
 
+enum class EBreakBehavior : uint8;
+
+struct FConcertSyncReplicationActivity;
 struct FConcertSyncReplicationPayload_Mute;
 struct FConcertSessionClientInfo;
 struct FConcertSyncReplicationPayload_LeaveReplication;
@@ -46,6 +50,9 @@ namespace UE::ConcertSyncServer::Replication
 		/** Gets the replication leave activity with ActivityId. */
 		virtual bool GetLeaveReplicationActivityById(const int64 ActivityId, FConcertSyncReplicationPayload_LeaveReplication& OutLeaveReplication) const = 0;
 
+		/** Enumerates all mute activities. */
+		virtual void EnumerateMuteActivities(TFunctionRef<EBreakBehavior(const FConcertSyncReplicationActivity& Activity)> Callback) const = 0;
+
 		virtual ~IReplicationWorkspace() = default;
 	};
 }
diff --git a/Engine/Plugins/Developer/Concert/ConcertSync/ConcertSyncTest/Source/ConcertSyncTest/Private/Replication/Restore/RestoreMuteState.spec.cpp b/Engine/Plugins/Developer/Concert/ConcertSync/ConcertSyncTest/Source/ConcertSyncTest/Private/Replication/Restore/RestoreMuteState.spec.cpp
new file mode 100644
--- /dev/null
+++ b/Engine/Plugins/Developer/Concert/ConcertSync/ConcertSyncTest/Source/ConcertSyncTest/Private/Replication/Restore/RestoreMuteState.spec.cpp
@@ -0,0 +1,122 @@
+﻿// Copyright Epic Games, Inc. All Rights Reserved.
+
+#include "Replication/Util/Mocks/ReplicationWorkspaceCallInterceptorMock.h"
+#include "Replication/Util/Spec/ReplicationClient.h"
+#include "Replication/Util/Spec/ReplicationServer.h"
+#include "Replication/Util/Spec/ObjectTestReplicator.h"
+
+#include "Misc/AutomationTest.h"
+#include "Replication/ReplicationTestInterface.h"
+#include "Util/ScopedSessionDatabase.h"
+
+namespace UE::ConcertSyncTests::Replication::RestoreContent
+{
+	BEGIN_DEFINE_SPEC(FRestoreMuteStateSpec, "Editor.Concert.Replication.RestoreContent.Mute", EAutomationTestFlags::EditorContext | EAutomationTestFlags::EngineFilter)
+		TUniquePtr<FScopedSessionDatabase> SessionDatabase;
+		TSharedPtr<ConcertSyncServer::Replication::IReplicationWorkspace> ReplicationWorkspace;
+		TUniquePtr<FObjectTestReplicator> ObjectReplicator;
+	
+		TUniquePtr<FReplicationServer> Server;
+		FReplicationClient* Client = nullptr;
+	
+		FReplicationClient& ConnectClient()
+		{
+			FReplicationClient& CreatedClient = Server->ConnectClient();
+			SessionDatabase->SetEndpoint(CreatedClient.GetEndpointId(), { CreatedClient.GetClientInfo() });
+			return CreatedClient;
+		}
+	END_DEFINE_SPEC(FRestoreMuteStateSpec);
+
+	/** This tests that mute state is correctly restored. */
+	void FRestoreMuteStateSpec::Define()
+	{
+		BeforeEach([this]
+		{
+			SessionDatabase = MakeUnique<FScopedSessionDatabase>(*this);
+			const auto FindSessionClient = [this](const FGuid& EndpointId) -> TOptional<FConcertSessionClientInfo>
+			{
+				FConcertSessionClientInfo Info;
+				const bool bFound = Server->GetServerSessionMock()->FindSessionClient(EndpointId, Info);
+				return bFound ? Info : TOptional<FConcertSessionClientInfo>{};
+			};
+			const auto ShouldIgnoreClientActivityOnRestore = [](const FGuid& EndpointId){ return false; };
+			ReplicationWorkspace = ConcertSyncServer::TestInterface::CreateReplicationWorkspace(*SessionDatabase, FindSessionClient, ShouldIgnoreClientActivityOnRestore);
+			
+			ObjectReplicator = MakeUnique<FObjectTestReplicator>();
+			Server = MakeUnique<FReplicationServer>(*this, EConcertSyncSessionFlags::Default_MultiUserSession, ReplicationWorkspace.ToSharedRef());
+			Client = &ConnectClient();
+
+			Client->JoinReplication(ObjectReplicator->CreateSenderArgs());
+			
+		});
+		AfterEach([this]
+		{
+			Server.Reset();
+			ObjectReplicator.Reset();
+			SessionDatabase.Reset();
+			ReplicationWorkspace.Reset();
+		});
+
+		Describe("If an object had been muted before rejoining session", [this]
+		{
+			BeforeEach([this]
+			{
+				IConcertClientReplicationManager& ReplicationManager = Client->GetClientReplicationManager();
+				
+				ReplicationManager.MuteObjects({ ObjectReplicator->TestObject });
+				ReplicationManager.LeaveReplicationSession();
+				ReplicationManager.JoinReplicationSession({});
+			});
+			
+			It("EConcertReplicationRestoreContentFlags::StreamsAndAuthority skips mute restoration", [this]
+			{
+				IConcertClientReplicationManager& ReplicationManager = Client->GetClientReplicationManager();
+				ReplicationManager.RestoreContent({ EConcertReplicationRestoreContentFlags::StreamsAndAuthority });
+				
+				bool bReceivedResponse = false;
+				ReplicationManager.QueryMuteState()
+					.Next([this, &bReceivedResponse](FConcertReplication_QueryMuteState_Response&& Response)
+					{
+						bReceivedResponse = true;
+						TestEqual(TEXT("ExplicitlyMutedObjects.Num()"), Response.ExplicitlyMutedObjects.Num(), 0);
+					});
+				TestTrue(TEXT("Received response"), bReceivedResponse);
+			});
+
+			It("EConcertReplicationRestoreContentFlags::All restores mute state", [this]
+			{
+				IConcertClientReplicationManager& ReplicationManager = Client->GetClientReplicationManager();
+				ReplicationManager.RestoreContent({ EConcertReplicationRestoreContentFlags::All });
+				
+				bool bReceivedResponse = false;
+				ReplicationManager.QueryMuteState()
+					.Next([this, &bReceivedResponse](FConcertReplication_QueryMuteState_Response&& Response)
+					{
+						bReceivedResponse = true;
+						TestEqual(TEXT("ExplicitlyMutedObjects.Num()"), Response.ExplicitlyMutedObjects.Num(), 1);
+						TestTrue(TEXT("ExplicitlyMutedObjects.Contains(TestObject)"), Response.ExplicitlyMutedObjects.Contains(ObjectReplicator->TestObject));
+					});
+				TestTrue(TEXT("Received response"), bReceivedResponse);
+			});
+		});
+
+		It("If object had been muted and then unmuted, then the object is not muted when rejoining", [this]
+		{
+			IConcertClientReplicationManager& ReplicationManager = Client->GetClientReplicationManager();
+			ReplicationManager.MuteObjects({ ObjectReplicator->TestObject });
+			ReplicationManager.UnmuteObjects({ ObjectReplicator->TestObject });
+			ReplicationManager.LeaveReplicationSession();
+			ReplicationManager.JoinReplicationSession({});
+			ReplicationManager.RestoreContent({ EConcertReplicationRestoreContentFlags::All });
+				
+			bool bReceivedResponse = false;
+			ReplicationManager.QueryMuteState()
+				.Next([this, &bReceivedResponse](FConcertReplication_QueryMuteState_Response&& Response)
+				{
+					bReceivedResponse = true;
+					TestEqual(TEXT("ExplicitlyMutedObjects.Num()"), Response.ExplicitlyMutedObjects.Num(), 0);
+				});
+			TestTrue(TEXT("Received response"), bReceivedResponse);
+		});
+	}
+}
diff --git a/Engine/Plugins/Developer/Concert/ConcertSync/ConcertSyncTest/Source/ConcertSyncTest/Private/Replication/Restore/RestoreMuteStateSyncControl.spec.cpp b/Engine/Plugins/Developer/Concert/ConcertSync/ConcertSyncTest/Source/ConcertSyncTest/Private/Replication/Restore/RestoreMuteStateSyncControl.spec.cpp
new file mode 100644
--- /dev/null
+++ b/Engine/Plugins/Developer/Concert/ConcertSync/ConcertSyncTest/Source/ConcertSyncTest/Private/Replication/Restore/RestoreMuteStateSyncControl.spec.cpp
@@ -0,0 +1,131 @@
+﻿// Copyright Epic Games, Inc. All Rights Reserved.
+
+#include "Replication/Util/Mocks/ReplicationWorkspaceCallInterceptorMock.h"
+#include "Replication/Util/Spec/ReplicationClient.h"
+#include "Replication/Util/Spec/ReplicationServer.h"
+#include "Replication/Util/Spec/ObjectTestReplicator.h"
+
+#include "Misc/AutomationTest.h"
+#include "Replication/ReplicationTestInterface.h"
+#include "Replication/Util/ConcertClientReplicationBridgeMock.h"
+#include "Util/ScopedSessionDatabase.h"
+
+namespace UE::ConcertSyncTests::Replication::RestoreContent
+{
+	BEGIN_DEFINE_SPEC(FRestoreMuteStateSyncControlSpec, "Editor.Concert.Replication.RestoreContent.Mute", EAutomationTestFlags::EditorContext | EAutomationTestFlags::EngineFilter)
+		TUniquePtr<FScopedSessionDatabase> SessionDatabase;
+		TSharedPtr<ConcertSyncServer::Replication::IReplicationWorkspace> ReplicationWorkspace;
+		TUniquePtr<FObjectTestReplicator> ObjectReplicator;
+	
+		TUniquePtr<FReplicationServer> Server;
+		FReplicationClient* SenderClient = nullptr;
+		FReplicationClient* ReceiverClient = nullptr;
+
+		FGuid StreamId = FGuid::NewGuid();
+	
+		FReplicationClient& ConnectClient()
+		{
+			FReplicationClient& CreatedClient = Server->ConnectClient();
+			SessionDatabase->SetEndpoint(CreatedClient.GetEndpointId(), { CreatedClient.GetClientInfo() });
+			return CreatedClient;
+		}
+	END_DEFINE_SPEC(FRestoreMuteStateSyncControlSpec);
+
+	/** This tests that sync control is correct when restoring mute state. */
+	void FRestoreMuteStateSyncControlSpec::Define()
+	{
+		BeforeEach([this]
+		{
+			SessionDatabase = MakeUnique<FScopedSessionDatabase>(*this);
+			const auto FindSessionClient = [this](const FGuid& EndpointId) -> TOptional<FConcertSessionClientInfo>
+			{
+				FConcertSessionClientInfo Info;
+				const bool bFound = Server->GetServerSessionMock()->FindSessionClient(EndpointId, Info);
+				return bFound ? Info : TOptional<FConcertSessionClientInfo>{};
+			};
+			const auto ShouldIgnoreClientActivityOnRestore = [](const FGuid& EndpointId){ return false; };
+			ReplicationWorkspace = ConcertSyncServer::TestInterface::CreateReplicationWorkspace(*SessionDatabase, FindSessionClient, ShouldIgnoreClientActivityOnRestore);
+			
+			ObjectReplicator = MakeUnique<FObjectTestReplicator>();
+			Server = MakeUnique<FReplicationServer>(*this, EConcertSyncSessionFlags::Default_MultiUserSession, ReplicationWorkspace.ToSharedRef());
+			SenderClient = &ConnectClient();
+			ReceiverClient = &ConnectClient();
+
+			SenderClient->JoinReplication(ObjectReplicator->CreateSenderArgs(StreamId));
+			ReceiverClient->JoinReplicationAsListener({ ObjectReplicator->TestObject });
+		});
+		AfterEach([this]
+		{
+			Server.Reset();
+			ObjectReplicator.Reset();
+			SessionDatabase.Reset();
+			ReplicationWorkspace.Reset();
+		});
+
+		It("If object had been muted, rejoining session does not grant sync control", [this]
+		{
+			IConcertClientReplicationManager& ReplicationManager = SenderClient->GetClientReplicationManager();
+			ReplicationManager.TakeAuthorityOver({ ObjectReplicator->TestObject });
+			ReplicationManager.MuteObjects({ ObjectReplicator->TestObject });
+			ReplicationManager.LeaveReplicationSession();
+			ReplicationManager.JoinReplicationSession({});
+
+			const FConcertObjectInStreamID ObjectInStreamID{ StreamId, ObjectReplicator->TestObject };
+			bool bReceivedResponse = false;
+			ReplicationManager.RestoreContent({ EConcertReplicationRestoreContentFlags::All })
+				.Next([this, &ObjectInStreamID, &bReceivedResponse](FConcertReplication_RestoreContent_Response&& Response)
+				{
+					bReceivedResponse = true;
+					TestEqual(TEXT("NewControlStates.Num()"), Response.SyncControl.NewControlStates.Num(), 1);
+					
+					const bool* bHasSyncControl = Response.SyncControl.NewControlStates.Find(ObjectInStreamID);
+					TestTrue(TEXT("Object does not have sync control"), bHasSyncControl && !*bHasSyncControl);
+				});
+			TestTrue(TEXT("Received response"), bReceivedResponse);
+			TestFalse(TEXT("!HasSyncControl(TestObject)"), ReplicationManager.HasSyncControl(ObjectInStreamID));
+		});
+		
+		Describe("If object had not been muted", [this]
+		{
+			BeforeEach([this]
+			{
+				IConcertClientReplicationManager& ReplicationManager = SenderClient->GetClientReplicationManager();
+				ReplicationManager.TakeAuthorityOver({ ObjectReplicator->TestObject });
+				ReplicationManager.MuteObjects({ ObjectReplicator->TestObject });
+				ReplicationManager.UnmuteObjects({ ObjectReplicator->TestObject });
+				ReplicationManager.LeaveReplicationSession();
+				ReplicationManager.JoinReplicationSession({});
+			});
+			
+			It("Rejoining session grant sync control", [this]
+			{
+				IConcertClientReplicationManager& ReplicationManager = SenderClient->GetClientReplicationManager();
+				const FConcertObjectInStreamID ObjectInStreamID{ StreamId, ObjectReplicator->TestObject };
+				
+				bool bReceivedResponse = false;
+				ReplicationManager
+					.RestoreContent({ EConcertReplicationRestoreContentFlags::All })
+					.Next([this, &ObjectInStreamID, &bReceivedResponse](FConcertReplication_RestoreContent_Response&& Response)
+					{
+						bReceivedResponse = true;
+						TestEqual(TEXT("SyncControl.IsEmpty()"), Response.SyncControl.NewControlStates.Num(), 1);
+
+						const bool* bHasSyncControl = Response.SyncControl.NewControlStates.Find(ObjectInStreamID);
+						TestTrue(TEXT("Object has sync control"), bHasSyncControl && *bHasSyncControl);
+					});
+				TestTrue(TEXT("Received response"), bReceivedResponse);
+				TestTrue(TEXT("HasSyncControl(TestObject)"), ReplicationManager.HasSyncControl(ObjectInStreamID));
+			});
+			
+			It("Can replicate object after rejoining", [this]
+			{
+				SenderClient->GetClientReplicationManager().RestoreContent({ EConcertReplicationRestoreContentFlags::All });
+
+				SenderClient->GetBridgeMock().InjectAvailableObject(*ObjectReplicator->TestObject);
+				ReceiverClient->GetBridgeMock().InjectAvailableObject(*ObjectReplicator->TestObject);
+				ObjectReplicator->SimulateSendObjectToReceiver(*this, { *SenderClient, *Server, *ReceiverClient }, { StreamId });
+				ObjectReplicator->TestValuesWereReplicated(*this);
+			});
+		});
+	}
+}
diff --git a/Engine/Plugins/Developer/Concert/ConcertSync/ConcertSyncTest/Source/ConcertSyncTest/Private/Replication/Util/Mocks/ReplicationWorkspaceCallInterceptorMock.h b/Engine/Plugins/Developer/Concert/ConcertSync/ConcertSyncTest/Source/ConcertSyncTest/Private/Replication/Util/Mocks/ReplicationWorkspaceCallInterceptorMock.h
--- a/Engine/Plugins/Developer/Concert/ConcertSync/ConcertSyncTest/Source/ConcertSyncTest/Private/Replication/Util/Mocks/ReplicationWorkspaceCallInterceptorMock.h
+++ b/Engine/Plugins/Developer/Concert/ConcertSync/ConcertSyncTest/Source/ConcertSyncTest/Private/Replication/Util/Mocks/ReplicationWorkspaceCallInterceptorMock.h
@@ -7,6 +7,7 @@
 
 #include "Misc/Optional.h"
 #include "Misc/Guid.h"
+#include "Templates/Function.h"
 #include "Templates/Tuple.h"
 
 namespace UE::ConcertSyncTests::Replication
@@ -32,6 +33,8 @@ namespace UE::ConcertSyncTests::Replication
 		TOptional<FConcertSyncReplicationPayload_LeaveReplication> ReturnResult_GetLastLeaveReplicationActivityByClient;
 		/** The result to return in GetLastLeaveReplicationActivityByClient. */
 		TOptional<FConcertSyncReplicationPayload_LeaveReplication> ReturnResult_GetLeaveReplicationActivityById;
+		/** The values to enumerate in EnumerateMuteActivities. */
+		TOptional<TArray<FConcertSyncReplicationActivity>> ReturnResult_EnumerateMuteActivities;
 		
 		//~ Begin IReplicationWorkspace Interface
 		virtual TOptional<int64> ProduceClientLeaveReplicationActivity(const FGuid& EndpointId, const FConcertSyncReplicationPayload_LeaveReplication& EventData) override
@@ -63,6 +66,21 @@ namespace UE::ConcertSyncTests::Replication
 			}
 			return ReturnResult_GetLeaveReplicationActivityById.IsSet();
 		}
+		virtual void EnumerateMuteActivities(TFunctionRef<EBreakBehavior(const FConcertSyncReplicationActivity& Activity)> Callback) const override
+		{
+			if (!ReturnResult_EnumerateMuteActivities)
+			{
+				return;
+			}
+			
+			for (const FConcertSyncReplicationActivity& Activity : *ReturnResult_EnumerateMuteActivities)
+			{
+				if (Callback(Activity) == EBreakBehavior::Break)
+				{
+					break;
+				}
+			}
+		}
 		//~ End IReplicationWorkspace Interface
 	};
 }
diff --git a/Engine/Plugins/Developer/Concert/ConcertSync/ConcertSyncTest/Source/ConcertSyncTest/Private/Replication/Util/Mocks/ReplicationWorkspaceEmptyMock.h b/Engine/Plugins/Developer/Concert/ConcertSync/ConcertSyncTest/Source/ConcertSyncTest/Private/Replication/Util/Mocks/ReplicationWorkspaceEmptyMock.h
--- a/Engine/Plugins/Developer/Concert/ConcertSync/ConcertSyncTest/Source/ConcertSyncTest/Private/Replication/Util/Mocks/ReplicationWorkspaceEmptyMock.h
+++ b/Engine/Plugins/Developer/Concert/ConcertSync/ConcertSyncTest/Source/ConcertSyncTest/Private/Replication/Util/Mocks/ReplicationWorkspaceEmptyMock.h
@@ -3,6 +3,7 @@
 #pragma once
 
 #include "Replication/IReplicationWorkspace.h"
+#include "Templates/Function.h"
 
 namespace UE::ConcertSyncTests::Replication
 {
@@ -15,6 +16,7 @@ namespace UE::ConcertSyncTests::Replication
 		virtual TOptional<int64> ProduceClientMuteReplicationActivity(const FGuid& EndpointId, const FConcertSyncReplicationPayload_Mute& EventData) override { return {}; }
 		virtual bool GetLastLeaveReplicationActivityByClient(const FConcertSessionClientInfo& InClientInfo, FConcertSyncReplicationPayload_LeaveReplication& OutLeaveReplication) const override { return false; }
 		virtual bool GetLeaveReplicationActivityById(const int64 ActivityId, FConcertSyncReplicationPayload_LeaveReplication& OutLeaveReplication) const override { return false; }
+		virtual void EnumerateMuteActivities(TFunctionRef<EBreakBehavior(const FConcertSyncReplicationActivity& Activity)> Callback) const override {}
 		//~ End IReplicationWorkspace Interface
 	};
 }
