Fix and simplify incorrect unit test

diff --git a/Engine/Plugins/Developer/Concert/ConcertSync/ConcertSyncTest/Source/ConcertSyncTest/Private/Database/ReplicationActivity.spec.cpp b/Engine/Plugins/Developer/Concert/ConcertSync/ConcertSyncTest/Source/ConcertSyncTest/Private/Database/ReplicationActivity.spec.cpp
--- a/Engine/Plugins/Developer/Concert/ConcertSync/ConcertSyncTest/Source/ConcertSyncTest/Private/Database/ReplicationActivity.spec.cpp
+++ b/Engine/Plugins/Developer/Concert/ConcertSync/ConcertSyncTest/Source/ConcertSyncTest/Private/Database/ReplicationActivity.spec.cpp
@@ -63,20 +63,6 @@ namespace UE::ConcertSyncTests
 			IFileManager::Get().DeleteDirectory(*GetDatabasePath(), false, true);
 			Database.Reset();
 		});
-		
-		struct FActivityEntry
-		{
-			FConcertSyncReplicationActivity Activity;
-			int64 ActivityId = INDEX_NONE;
-			int64 ReplicationEventId = INDEX_NONE;
-				
-			FActivityEntry(FConcertSyncReplicationActivity Activity) : Activity(MoveTemp(Activity)) {}
-
-			bool operator==(const FActivityEntry& Right) const
-			{
-				return ActivityId == Right.ActivityId && ReplicationEventId == Right.ReplicationEventId && Activity.EventData == Right.Activity.EventData;
-			}
-		};
 
 		It("Add and retrieve FConcertSyncReplicationPayload_LeaveReplication", [this]
 		{
@@ -115,32 +101,34 @@ namespace UE::ConcertSyncTests
 		It("EnumerateReplicationActivities", [this]
 		{
 			const FGuid StreamId = FGuid::NewGuid();
-			TArray<FActivityEntry> Activities_OriginalOrder =
+			TArray Activities_OriginalOrder =
 			{
 				FConcertSyncReplicationActivity(MakeLeaveReplicationActivity(FSoftObjectPath(TEXT("/Game/Maps.Map:PersistentLevel.Cube.StaticMeshComponent0")), StreamId)),
 				FConcertSyncReplicationActivity(MakeLeaveReplicationActivity(FSoftObjectPath(TEXT("/Game/Maps.Map:PersistentLevel.Cube.StaticMeshComponent1")), StreamId)),
 				FConcertSyncReplicationActivity(MakeLeaveReplicationActivity(FSoftObjectPath(TEXT("/Game/Maps.Map:PersistentLevel.Cube.StaticMeshComponent2")), StreamId))
 			};
-			for (FActivityEntry& Entry : Activities_OriginalOrder)
+			for (FConcertSyncReplicationActivity& Activity : Activities_OriginalOrder)
 			{
-				Database->AddReplicationActivity(Entry.Activity, Entry.ActivityId, Entry.ReplicationEventId);
+				Database->AddReplicationActivity(Activity, Activity.ActivityId, Activity.EventId);
 			}
 			
-			TArray<FActivityEntry> Activities_EnumerationOrder;
+			TArray<FConcertSyncReplicationActivity> Activities_EnumerationOrder;
 			Database->EnumerateReplicationActivities([&Activities_EnumerationOrder](FConcertSyncReplicationActivity&& Activity)
 			{
 				Activities_EnumerationOrder.Emplace(MoveTemp(Activity));
 				return true;
 			});
 
-			if (Activities_EnumerationOrder.Num() && Activities_OriginalOrder.Num())
+			if (Activities_EnumerationOrder.Num() != Activities_OriginalOrder.Num())
 			{
 				AddError(TEXT("Counts do not match"));
 				return;
 			}
 			for (int32 i = 0; i < Activities_EnumerationOrder.Num(); ++i)
 			{
-				TestEqual(TEXT("Activities match"), Activities_EnumerationOrder[i], Activities_OriginalOrder[i]);
+				TestEqual(TEXT("ActivityId"), Activities_EnumerationOrder[i].ActivityId, Activities_OriginalOrder[i].ActivityId);
+				TestEqual(TEXT("EventId"), Activities_EnumerationOrder[i].EventId, Activities_OriginalOrder[i].EventId);
+				TestEqual(TEXT("EventData"), Activities_EnumerationOrder[i].EventData, Activities_OriginalOrder[i].EventData);
 			}
 		});
 		
@@ -149,15 +137,15 @@ namespace UE::ConcertSyncTests
 			const FGuid StreamId = FGuid::NewGuid();
 			const FGuid ClientOne = FGuid::NewGuid();
 			const FGuid ClientTwo = FGuid::NewGuid();
-			TArray<FActivityEntry> Activities_OriginalOrder =
+			TArray Activities_OriginalOrder =
 			{
 				MakeActivity(MakeLeaveReplicationActivity(FSoftObjectPath(TEXT("/Game/Maps.Map:PersistentLevel.Cube.StaticMeshComponent0")), StreamId), ClientOne),
 				MakeActivity(MakeLeaveReplicationActivity(FSoftObjectPath(TEXT("/Game/Maps.Map:PersistentLevel.Cube.StaticMeshComponent1")), StreamId), ClientOne),
 				MakeActivity(MakeLeaveReplicationActivity(FSoftObjectPath(TEXT("/Game/Maps.Map:PersistentLevel.Cube.StaticMeshComponent2")), StreamId), ClientTwo)
 			};
-			for (FActivityEntry& Entry : Activities_OriginalOrder)
+			for (FConcertSyncReplicationActivity& Activity : Activities_OriginalOrder)
 			{
-				Database->AddReplicationActivity(Entry.Activity, Entry.ActivityId, Entry.ReplicationEventId);
+				Database->AddReplicationActivity(Activity, Activity.ActivityId, Activity.EventId);
 			}
 			
 			int64 EventId_ClientOne = INDEX_NONE;
@@ -168,9 +156,9 @@ namespace UE::ConcertSyncTests
 			const bool bSuccess_InvalidQuery = Database->GetReplicationMaxEventIdByClientAndType(FGuid::NewGuid(), EConcertSyncReplicationActivityType::LeaveReplication, EventId_Invalid);
 			
 			TestTrue(TEXT("Query 1"), bSuccess_Query1);
-			TestEqual(TEXT("Client 1 Max EventId"), EventId_ClientOne, Activities_OriginalOrder[1].ReplicationEventId);
+			TestEqual(TEXT("Client 1 Max EventId"), EventId_ClientOne, Activities_OriginalOrder[1].EventId);
 			TestTrue(TEXT("Query 2"), bSuccess_Query2);
-			TestEqual(TEXT("Client 2 Max EventId"), EventId_ClientTwo, Activities_OriginalOrder[2].ReplicationEventId);
+			TestEqual(TEXT("Client 2 Max EventId"), EventId_ClientTwo, Activities_OriginalOrder[2].EventId);
 			TestTrue(TEXT("Invalid query"), bSuccess_InvalidQuery);
 			// Sql query returns NULL here (since nothing is found).
 			// Under the hood, sqlite3_value_int64 is used to get the value of the returned column... NULL is converted to 0 by sqlite3_value_int64.
