Moved stream unit tests to Replication/Stream folder

diff --git a/Engine/Plugins/Developer/Concert/ConcertSync/ConcertSyncTest/Source/ConcertSyncTest/Private/Replication/Stream/StreamRequestTests_Change.cpp b/Engine/Plugins/Developer/Concert/ConcertSync/ConcertSyncTest/Source/ConcertSyncTest/Private/Replication/Stream/StreamRequestTests_Change.cpp
new file mode 100644
--- /dev/null
+++ b/Engine/Plugins/Developer/Concert/ConcertSync/ConcertSyncTest/Source/ConcertSyncTest/Private/Replication/Stream/StreamRequestTests_Change.cpp
@@ -0,0 +1,297 @@
+// Copyright Epic Games, Inc. All Rights Reserved.
+
+#include "Util/ClientServerCommunicationTest.h"
+
+#include "Replication/Data/ReplicationStream.h"
+#include "Replication/IConcertClientReplicationManager.h"
+#include "Replication/Util/ChangeStreamsTestBase.h"
+#include "Replication/Util/SendReceiveObjectTestBase.h"
+#include "Replication/TestReflectionObject.h"
+
+#include "Misc/AutomationTest.h"
+#include "UObject/UObjectGlobals.h"
+#include "UObject/Package.h"
+
+namespace UE::ConcertSyncTests::Replication::Stream
+{
+	/** Simple case of FConcertChangeStream_Request where the requesting client has authority and thus cannot generate any authority conflicts. */
+	IMPLEMENT_CUSTOM_SIMPLE_AUTOMATION_TEST(FSimpleChangeStream, FChangeStreamsTestBase, "Editor.Concert.Replication.Stream.SimpleChangeStream", EAutomationTestFlags::EditorContext | EAutomationTestFlags::EngineFilter);
+	bool FSimpleChangeStream::RunTest(const FString& Parameters)
+	{
+		using namespace ConcertSyncClient::Replication;
+
+		// 1. Start up client with float property stream
+		UTestReflectionObject* TestObject = NewObject<UTestReflectionObject>(GetTransientPackage());
+		auto[InitialStreamId, InitialStream] = CreateFloatPropertyStream(*TestObject);
+		SenderArgs.Streams = { InitialStream };
+		SetUpClientAndServer();
+
+		// 2.1 Modify existing stream
+		FConcertPropertySelection& ChangedInitialProperties = GetPropertySelection(InitialStream, *TestObject);
+		AddVectorProperty(ChangedInitialProperties);
+		FConcertReplication_ChangeStream_Request ModifyRequest;
+		ModifyRequest.ObjectsToPut.Add({ InitialStreamId, TestObject }, { ChangedInitialProperties });
+		ChangeStreamForSenderClientAndValidate(TEXT("ModifyRequest"), ModifyRequest, { InitialStream.BaseDescription });
+
+		// 2.2 Add new stream
+		auto[DynamicStreamId, DynamicStream] = CreateFloatPropertyStream(*TestObject);
+		FConcertReplication_ChangeStream_Request DynamicCreationRequest;
+		DynamicCreationRequest.StreamsToAdd.Add(DynamicStream);
+		ChangeStreamForSenderClientAndValidate(TEXT("DynamicCreationRequest"), DynamicCreationRequest, { InitialStream.BaseDescription, DynamicStream.BaseDescription });
+
+		// 2.3 Remove stream
+		FConcertReplication_ChangeStream_Request RemoveStreamRequest;
+		RemoveStreamRequest.StreamsToRemove.Add(InitialStreamId);
+		ChangeStreamForSenderClientAndValidate(TEXT("RemoveStreamRequest"), RemoveStreamRequest, { DynamicStream.BaseDescription });
+
+		// 2.4 Remove object > removes stream implicitly
+		FConcertReplication_ChangeStream_Request RemoveObjectRequest;
+		RemoveObjectRequest.ObjectsToRemove.Add({ DynamicStreamId, TestObject });
+		ChangeStreamForSenderClientAndValidate(TEXT("RemoveObjectRequest"), RemoveObjectRequest, {});
+		
+		return true;
+	}
+
+	/** Case of FConcertChangeStream_Request in which requester has authority over an object and creates two streams which have authority over overlapping properties. This is allowed. */
+	IMPLEMENT_CUSTOM_SIMPLE_AUTOMATION_TEST(FChangeStreamWithOverlappingProperties, FChangeStreamsTestBase, "Editor.Concert.Replication.Stream.ChangeStreamWithOverlappingProperties", EAutomationTestFlags::EditorContext | EAutomationTestFlags::EngineFilter);
+	bool FChangeStreamWithOverlappingProperties::RunTest(const FString& Parameters)
+	{
+		using namespace ConcertSyncClient::Replication;
+
+		// 1. Start client without any streams
+		UTestReflectionObject* TestObject = NewObject<UTestReflectionObject>(GetTransientPackage());
+		SetUpClientAndServer();
+		
+		// 2.1 Add stream with float property
+		auto[FloatStreamID, FloatStream] = CreateFloatPropertyStream(*TestObject);
+		FConcertReplication_ChangeStream_Request CreateFloatStreamRequest;
+		CreateFloatStreamRequest.StreamsToAdd.Add(FloatStream);
+		ChangeStreamForSenderClientAndValidate(TEXT("CreateFloatStreamRequest"), CreateFloatStreamRequest, { FloatStream.BaseDescription });
+		// 2.2 Add a stream with vector property
+		auto[VectorFloatStreamID, VectorFloatStream] = CreateVectorPropertyStream(*TestObject);
+		FConcertReplication_ChangeStream_Request CreateVectorFloatStreamRequest; 
+		CreateVectorFloatStreamRequest.StreamsToAdd.Add(VectorFloatStream);
+		ChangeStreamForSenderClientAndValidate(TEXT("CreateVectorFloatStreamRequest"), CreateVectorFloatStreamRequest, { FloatStream.BaseDescription, VectorFloatStream.BaseDescription });
+
+		// 2.3 Take authority over both streams
+		bool bTookAuthority = false;
+		FConcertReplication_ChangeAuthority_Request TakeAuthorityRequest;
+		TakeAuthorityRequest.TakeAuthority.Add(TestObject, FConcertStreamArray{{ FloatStreamID, VectorFloatStreamID }});
+		ClientReplicationManager_Sender->RequestAuthorityChange({ TakeAuthorityRequest })
+			.Next([this, &bTookAuthority](FConcertReplication_ChangeAuthority_Response&& Response)
+			{
+				bTookAuthority = true;
+				TestTrue(TEXT("ErrorCode == Handled"), Response.ErrorCode == EReplicationResponseErrorCode::Handled);
+				TestTrue(TEXT("Authority request > Success"), Response.RejectedObjects.IsEmpty());
+			});
+		TestTrue(TEXT("Authority request > Received response"), bTookAuthority);
+
+		// 2.4 Appending the overlapping float property works
+		FConcertReplication_ChangeStream_Request AppendFloatRequest;
+		FConcertPropertySelection& NewSelection = GetPropertySelection(VectorFloatStream, *TestObject);
+		AddFloatProperty(NewSelection);
+		AppendFloatRequest.ObjectsToPut.Add(FConcertObjectInStreamID{ VectorFloatStreamID, TestObject }, FConcertReplication_ChangeStream_PutObject{ NewSelection });
+		ChangeStreamForSenderClientAndValidate(TEXT("AppendFloatRequest"), AppendFloatRequest, { FloatStream.BaseDescription, VectorFloatStream.BaseDescription });
+		
+		return true;
+	}
+
+	/**
+	 * The requester and client have authority over separate properties on the same object.
+	 * The request would cause the requester's existing authority to overlap with the other client, which is rejected.
+	 */
+	IMPLEMENT_CUSTOM_SIMPLE_AUTOMATION_TEST(FRejectChangeStreamWithConflictingAuthority, FChangeStreamsTestBase, "Editor.Concert.Replication.Stream.RejectChangeStreamWithConflictingAuthority", EAutomationTestFlags::EditorContext | EAutomationTestFlags::EngineFilter);
+	bool FRejectChangeStreamWithConflictingAuthority::RunTest(const FString& Parameters)
+	{
+		using namespace ConcertSyncClient::Replication;
+
+		// 1. Start client without any streams
+		UTestReflectionObject* TestObject = NewObject<UTestReflectionObject>(GetTransientPackage());
+		SetUpClientAndServer();
+		
+		// 2.1 Add stream with float property
+		auto[SenderStreamID, SenderStream] = CreateFloatPropertyStream(*TestObject);
+		FConcertReplication_ChangeStream_Request CreateFloatStreamRequest;
+		CreateFloatStreamRequest.StreamsToAdd.Add(SenderStream);
+		ChangeStreamForSenderClientAndValidate(TEXT("CreateFloatStreamRequest"), CreateFloatStreamRequest, { SenderStream.BaseDescription });
+
+		// 2.2 Have Sender take authority over the float property
+		bool bSenderTookAuthority = false;
+		ClientReplicationManager_Sender->TakeAuthorityOver({ TestObject })
+			.Next([this, &bSenderTookAuthority](const FConcertReplication_ChangeAuthority_Response& Response) mutable
+			{
+				bSenderTookAuthority = true;
+				TestTrue(TEXT("ErrorCode == Handled"), Response.ErrorCode == EReplicationResponseErrorCode::Handled);
+				TestEqual(TEXT("Sender > No rejection taking authority"), Response.RejectedObjects.Num(), 0); 
+			});
+		TestTrue(TEXT("Sender > Received response taking authority"), bSenderTookAuthority);
+
+		// 2.3 Have Receiver create vector stream
+		auto[ReceiverStreamID, ReceiverStream] = CreateVectorPropertyStream(*TestObject);
+		FConcertReplication_ChangeStream_Request CreateVectorStreamRequest; 
+		CreateVectorStreamRequest.StreamsToAdd.Add(ReceiverStream);
+		bool bAddedVectorStream = false;
+		ClientReplicationManager_Receiver->ChangeStream(CreateVectorStreamRequest)
+			.Next([this, &bAddedVectorStream](FConcertReplication_ChangeStream_Response&& Response)
+			{
+				bAddedVectorStream = true;
+				TestTrue(TEXT("ErrorCode == Handled"), Response.ErrorCode == EReplicationResponseErrorCode::Handled);
+				TestTrue(TEXT("Vector Stream > Success"), Response.IsSuccess());
+			});
+		TestTrue(TEXT("Vector Stream> Received response"), bAddedVectorStream);
+		// For simplicity, we don't do any further validation that other clients can see the change
+		
+		// 2.4 Have Receiver take authority over vector stream
+		bool bReceiverTookAuthority = false;
+		ClientReplicationManager_Receiver->TakeAuthorityOver({ TestObject })
+			.Next([this, &bReceiverTookAuthority](const FConcertReplication_ChangeAuthority_Response& Response) mutable
+			{
+				bReceiverTookAuthority = true;
+				TestTrue(TEXT("ErrorCode == Handled"), Response.ErrorCode == EReplicationResponseErrorCode::Handled);
+				TestEqual(TEXT("Receiver > No rejection taking authority"), Response.RejectedObjects.Num(), 0); 
+			});
+		TestTrue(TEXT("Receiver > Received response taking authority"), bReceiverTookAuthority);
+
+		// 2.5 Receiver is rejected from modifying the stream to include the Float property because the Sender has authority over it
+		const FConcertObjectInStreamID TestObjectInReceiverStreamId { ReceiverStreamID, TestObject };
+		FConcertReplication_ChangeStream_Request AddFloatToStreamRequest;
+		FConcertPropertySelection NewSelection = GetPropertySelection(ReceiverStream, *TestObject);
+		AddFloatProperty(NewSelection);
+		AddFloatToStreamRequest.ObjectsToPut.Add(TestObjectInReceiverStreamId, { NewSelection });
+		bool bReceivedResponseAddingFloat = false;
+		
+		// Server logs a warning when rejecting - avoid the test being marked with a warning.
+		AddExpectedError(TEXT("Rejecting ChangeStream"));
+		ClientReplicationManager_Receiver->ChangeStream(AddFloatToStreamRequest)
+			.Next([this, TestObject, TheSenderStreamID = SenderStreamID, &TestObjectInReceiverStreamId, &bReceivedResponseAddingFloat](FConcertReplication_ChangeStream_Response&& Response)
+			{
+				bReceivedResponseAddingFloat = true;
+				TestTrue(TEXT("ErrorCode == Handled"), Response.ErrorCode == EReplicationResponseErrorCode::Handled);
+				TestEqual(TEXT("Append Float > 1 conflict"), Response.AuthorityConflicts.Num(), 1);
+				if (const FConcertReplicatedObjectId* ConflictingObject = Response.AuthorityConflicts.Find(TestObjectInReceiverStreamId))
+				{
+					TestEqual(TEXT("Append float > Conflict > Sender Stream correct"), ConflictingObject->StreamId, TheSenderStreamID);
+					TestEqual(TEXT("Append float > Conflict > Object correct"), ConflictingObject->Object, FSoftObjectPath(TestObject));
+					TestEqual(TEXT("Append float > Conflict > Endpoint ID correct"), ConflictingObject->SenderEndpointId, Client_Sender->ClientSessionMock->GetSessionClientEndpointId());
+				}
+				else
+				{
+					AddError(TEXT("Expected to find an authority object conflict for TestObject"));
+				}
+				TestTrue(TEXT("Append Float > Failure"), Response.IsFailure());
+			});
+		TestTrue(TEXT("Append Float > Received response"), bReceivedResponseAddingFloat);
+		
+		return true;  
+	}
+
+	/** Tests that no changes are made if a sub-step of FConcertChangeStream_Request causes an error. */
+	IMPLEMENT_CUSTOM_SIMPLE_AUTOMATION_TEST(FChangeStreamIsAtmoic, FChangeStreamsTestBase, "Editor.Concert.Replication.Stream.ChangeStreamIsAtmoic", EAutomationTestFlags::EditorContext | EAutomationTestFlags::EngineFilter);
+	bool FChangeStreamIsAtmoic::RunTest(const FString& Parameters)
+	{
+		using namespace ConcertSyncClient::Replication;
+
+		// 1. Start client without any streams
+		UTestReflectionObject* TestObject = NewObject<UTestReflectionObject>(GetTransientPackage());
+		SetUpClientAndServer();
+		
+		// 2.1 Add stream with float property
+		auto[StreamId, Stream] = CreateFloatPropertyStream(*TestObject);
+		FConcertReplication_ChangeStream_Request CreateStreamRequest;
+		CreateStreamRequest.StreamsToAdd.Add(Stream);
+		ChangeStreamForSenderClientAndValidate(TEXT("CreateStreamRequest"), CreateStreamRequest, { Stream.BaseDescription });
+
+		// 2.2 Make a request that will fail and check that no changes were made to the original stream
+		FConcertReplication_ChangeStream_Request InvalidRequest;
+		FConcertPropertySelection NewSelection = GetPropertySelection(Stream, *TestObject);
+		AddFloatProperty(NewSelection);
+		InvalidRequest.ObjectsToPut.Add(FConcertObjectInStreamID{ StreamId, TestObject}, FConcertReplication_ChangeStream_PutObject{ NewSelection });
+		InvalidRequest.StreamsToAdd.Add(Stream); // This will make it fail due to pre-existing stream ID
+
+		// Server logs a warning when rejecting - avoid the test being marked with a warning.
+		AddExpectedError(TEXT("Rejecting ChangeStream"));
+		
+		bool bReceivedChangeStreamResponse = false;
+		ClientReplicationManager_Sender->ChangeStream(InvalidRequest)
+			.Next([this, &bReceivedChangeStreamResponse](FConcertReplication_ChangeStream_Response&& Response)
+			{
+				bReceivedChangeStreamResponse = true;
+				TestTrue(TEXT("ErrorCode == Handled"), Response.ErrorCode == EReplicationResponseErrorCode::Handled);
+				TestTrue(TEXT("Invalid Request > Modify Stream > Failure"), Response.IsFailure());
+			});
+		
+		TestTrue(TEXT("Invalid Request > Modify Stream > Received response"), bReceivedChangeStreamResponse);
+		ValidateSenderClientStreams(TEXT("Invalid Request"), { Stream.BaseDescription });
+		
+		return true;
+	}
+
+	/** Removing an object from a stream over which the requester has authority, clears the authority the client had. Tests that another client can now claim authority over the previously authored properties. */
+	IMPLEMENT_CUSTOM_SIMPLE_AUTOMATION_TEST(FRemovingObjectClearsAuthority, FChangeStreamsTestBase, "Editor.Concert.Replication.Stream.RemovingObjectClearsAuthority", EAutomationTestFlags::EditorContext | EAutomationTestFlags::EngineFilter);
+	bool FRemovingObjectClearsAuthority::RunTest(const FString& Parameters)
+	{
+		using namespace ConcertSyncClient::Replication;
+
+		// 1. Start clients with equivalent streams
+		UTestReflectionObject* TestObject = NewObject<UTestReflectionObject>(GetTransientPackage());
+		auto[SenderStreamId, SenderStream] = CreateFloatPropertyStream(*TestObject);
+		auto[ReceiverStreamId, ReceiverStream] = CreateFloatPropertyStream(*TestObject);
+		SenderArgs.Streams = { SenderStream };
+		ReceiverArgs.Streams = { ReceiverStream };
+		SetUpClientAndServer();
+		
+		// 2.1 Sender takes authority
+		ClientReplicationManager_Sender->TakeAuthorityOver({ TestObject })
+			.Next([this](FConcertReplication_ChangeAuthority_Response&& Response)
+			{
+				TestTrue(TEXT("ErrorCode == Handled"), Response.ErrorCode == EReplicationResponseErrorCode::Handled);
+				TestTrue(TEXT("Sender > Success taking authority"), Response.RejectedObjects.IsEmpty());
+			});
+
+		// 2.2 Sender removes the object
+		FConcertReplication_ChangeStream_Request RemoveObjectRequest;
+		RemoveObjectRequest.ObjectsToRemove.Add({ SenderStreamId, TestObject });
+		ChangeStreamForSenderClientAndValidate(TEXT("Remove Object"), RemoveObjectRequest, {});
+
+		// 2.3 Receiver can now take authority since 2.2 implicitly removed authority
+		ClientReplicationManager_Receiver->TakeAuthorityOver({ TestObject })
+			.Next([this](FConcertReplication_ChangeAuthority_Response&& Response)
+			{
+				TestTrue(TEXT("ErrorCode == Handled"), Response.ErrorCode == EReplicationResponseErrorCode::Handled);
+				TestTrue(TEXT("Receiver > Success taking authority"), Response.RejectedObjects.IsEmpty());
+			});
+		
+		return true;
+	}
+
+	/**
+     * Tests that client updates its local cache of the server state when RequestAuthorityChange times out.
+     */
+    IMPLEMENT_CUSTOM_SIMPLE_AUTOMATION_TEST(FChangingStreamTimeoutRetainsServerState, FChangeStreamsTestBase, "Editor.Concert.Replication.Stream.ChangeStreamTimeoutRetainsServerState", EAutomationTestFlags::EditorContext | EAutomationTestFlags::EngineFilter);
+    bool FChangingStreamTimeoutRetainsServerState::RunTest(const FString& Parameters)
+    {
+    	// 1. Start up client with float property stream
+    	UTestReflectionObject* TestObject = NewObject<UTestReflectionObject>(GetTransientPackage());
+    	auto[InitialStreamId, InitialStream] = CreateFloatPropertyStream(*TestObject);
+    	SenderArgs.Streams = { InitialStream };
+    	SetUpClientAndServer();
+		ServerSession->SetTestFlags(EServerSessionTestingFlags::AllowRequestTimeouts);
+
+    	// 2. Timeout request
+    	ServerSession->UnregisterCustomRequestHandler<FConcertReplication_ChangeStream_Request>();
+    	FConcertReplication_ChangeStream_Request Request;
+    	Request.ObjectsToRemove.Add({ InitialStreamId, TestObject });
+    	ClientReplicationManager_Sender->ChangeStream(Request);
+
+    	// 3. Check that the local client's server prediction was reverted
+    	TArray<FConcertReplicationStream> Streams = ClientReplicationManager_Sender->GetRegisteredStreams();
+    	if (Streams.IsEmpty())
+    	{
+    		AddError(TEXT("Stream change was not reverted"));
+    		return true;
+    	}
+    	TestTrue(TEXT("Timed out Stream change reverted"), Streams[0].BaseDescription.ReplicationMap.ReplicatedObjects.Contains(TestObject));
+    	
+    	return true;
+    }
+}
diff --git a/Engine/Plugins/Developer/Concert/ConcertSync/ConcertSyncTest/Source/ConcertSyncTest/Private/Replication/Stream/StreamRequestTests_Frequency.cpp b/Engine/Plugins/Developer/Concert/ConcertSync/ConcertSyncTest/Source/ConcertSyncTest/Private/Replication/Stream/StreamRequestTests_Frequency.cpp
new file mode 100644
--- /dev/null
+++ b/Engine/Plugins/Developer/Concert/ConcertSync/ConcertSyncTest/Source/ConcertSyncTest/Private/Replication/Stream/StreamRequestTests_Frequency.cpp
@@ -0,0 +1,381 @@
+// Copyright Epic Games, Inc. All Rights Reserved.
+
+#include "Replication/Data/ReplicationStream.h"
+#include "Replication/IConcertClientReplicationManager.h"
+#include "Replication/Util/ChangeStreamsTestBase.h"
+#include "Replication/Util/SendReceiveObjectTestBase.h"
+#include "Replication/TestReflectionObject.h"
+
+#include "Misc/AutomationTest.h"
+#include "UObject/UObjectGlobals.h"
+#include "UObject/Package.h"
+
+namespace UE::ConcertSyncTests::Replication::Stream
+{
+	/**
+	 * Tests that frequency settings
+	 * - can be changed using all possible options (put, add, remove)
+	 * - the operation order is put, then add, then remove
+	 * - querying frequency info works
+	 * - replication rate cannot be set to 0
+	 * - atomicity: failing to change replication frequency fails all other changes
+	 * - frequency can only be overriden if object is in object replication map
+	 * - frequency override is removed if object is removed from object replication map
+	 */
+	IMPLEMENT_CUSTOM_SIMPLE_AUTOMATION_TEST(FChangeFrequency, FChangeStreamsTestBase, "Editor.Concert.Replication.Stream.Frequency.ChangeFrequency", EAutomationTestFlags::EditorContext | EAutomationTestFlags::EngineFilter);
+	bool FChangeFrequency::RunTest(const FString& Parameters)
+	{
+		// 1. Start up client with float property stream
+		UTestReflectionObject* MainTestObject = NewObject<UTestReflectionObject>(GetTransientPackage());
+		UTestReflectionObject* SecondaryTestObject = NewObject<UTestReflectionObject>(GetTransientPackage());
+		FGuid FloatStreamId;
+		FConcertReplicationStream FloatStream;
+		Tie(FloatStreamId, FloatStream) = CreateFloatPropertyStream(*MainTestObject);
+		SenderArgs.Streams = { FloatStream };
+		
+		SetUpClientAndServer();
+		// We'll intentionally be causing errors below that cause a warning to be logged.
+		AddExpectedError(TEXT("Rejecting ChangeStream"), EAutomationExpectedErrorFlags::Contains, 2);
+
+		FConcertBaseStreamInfo ExpectedStream = FloatStream.BaseDescription;
+		
+		// Change defaults to be SpecifiedRate at 42 update rate.
+		{
+			FConcertReplication_ChangeStream_Request ChangeRequest;
+			FConcertReplication_ChangeStream_Frequency FrequencyChange { .NewDefaults = { EConcertObjectReplicationMode::SpecifiedRate, 42 }, .Flags = EConcertReplicationChangeFrequencyFlags::SetDefaults };
+			ChangeRequest.FrequencyChanges.Add(FloatStreamId, FrequencyChange);
+			
+			ExpectedStream.FrequencySettings.Defaults = { EConcertObjectReplicationMode::SpecifiedRate, 42 };
+			ChangeStreamForSenderClientAndValidate(TEXT("SpecifiedRate with 42"), ChangeRequest, { ExpectedStream });
+		}
+		
+		// Not specifying update flags does not update default replication state.
+		{
+			FConcertReplication_ChangeStream_Request ChangeRequest;
+			FConcertReplication_ChangeStream_Frequency FrequencyChange { .NewDefaults = { EConcertObjectReplicationMode::SpecifiedRate, 1 } };
+			ChangeRequest.FrequencyChanges.Add(FloatStreamId, FrequencyChange);
+			ChangeStreamForSenderClientAndValidate(TEXT("Do not update defaults without update flags"), ChangeRequest, { ExpectedStream });
+		}
+		// Cannot set SpecifiedRate at 0 update rate.
+		{
+			FConcertReplication_ChangeStream_Request ChangeRequest;
+			FConcertReplication_ChangeStream_Frequency FrequencyChange { .NewDefaults = { EConcertObjectReplicationMode::SpecifiedRate, 0 }, .Flags = EConcertReplicationChangeFrequencyFlags::SetDefaults };
+			ChangeRequest.FrequencyChanges.Add(FloatStreamId, FrequencyChange);
+			ChangeStreamForSenderClientAndValidate(
+				TEXT("Defaults > Rate 0"),
+				ChangeRequest,
+				{ ExpectedStream },
+				EChangeStreamValidation::TestWasHandled | EChangeStreamValidation::TestResponseFailed
+				)
+				.Next([this, FloatStreamId](FConcertReplication_ChangeStream_Response&& Response)
+				{
+					const EConcertChangeStreamFrequencyErrorCode* ErrorCode = Response.FrequencyErrors.DefaultFailures.Find(FloatStreamId);
+					TestTrue(TEXT("Defaults > Rate 0 > Rejection code"), ErrorCode && *ErrorCode == EConcertChangeStreamFrequencyErrorCode::InvalidReplicationRate);
+				});
+		}
+		
+		// Failing frequency update reverts other changes in request (atomicity)
+		{
+			FConcertReplication_ChangeStream_Request ChangeRequest;
+			
+			FConcertReplication_ChangeStream_PutObject PutObject;
+			AddFloatProperty(PutObject.Properties);
+			PutObject.ClassPath = SecondaryTestObject->GetClass();
+			ChangeRequest.ObjectsToPut.Add({ FloatStreamId, SecondaryTestObject }, PutObject);
+			
+			FConcertReplication_ChangeStream_Frequency FrequencyChange;
+			FrequencyChange.OverridesToAdd.Add(SecondaryTestObject, { EConcertObjectReplicationMode::SpecifiedRate, 0 });
+			ChangeRequest.FrequencyChanges.Add(FloatStreamId, FrequencyChange);
+			
+			ChangeStreamForSenderClientAndValidate(
+				TEXT("Failing frequency update reverts other changes in request (atomicity)"),
+				ChangeRequest,
+				{ ExpectedStream },
+				EChangeStreamValidation::TestWasHandled | EChangeStreamValidation::TestResponseFailed
+				);
+		}
+		
+		// Add object and change frequency in same request
+		{
+			FConcertReplication_ChangeStream_Request ChangeRequest;
+			
+			FConcertReplication_ChangeStream_PutObject PutObject;
+			AddFloatProperty(PutObject.Properties);
+			PutObject.ClassPath = SecondaryTestObject->GetClass();
+			ChangeRequest.ObjectsToPut.Add({ FloatStreamId, SecondaryTestObject }, PutObject);
+			
+			FConcertReplication_ChangeStream_Frequency FrequencyChange;
+			FrequencyChange.OverridesToAdd.Add(SecondaryTestObject, { EConcertObjectReplicationMode::SpecifiedRate, 101 });
+			ChangeRequest.FrequencyChanges.Add(FloatStreamId, FrequencyChange);
+
+			FConcertReplicatedObjectInfo& ExpectedSecondary = ExpectedStream.ReplicationMap.ReplicatedObjects.Add(SecondaryTestObject);
+			AddFloatProperty(ExpectedSecondary.PropertySelection);
+			ExpectedSecondary.ClassPath = SecondaryTestObject->GetClass();
+			ExpectedStream.FrequencySettings.ObjectOverrides.Add(SecondaryTestObject, { EConcertObjectReplicationMode::SpecifiedRate, 101 });
+			ChangeStreamForSenderClientAndValidate(TEXT("Register object and set frequency"), ChangeRequest, { ExpectedStream });
+		}
+		
+		return true;
+	}
+
+	/**
+	 * Tests that requests for unregistered objects rejected:
+	 * - Cannot change frequency for unknown stream
+	 * - Cannot change frequency for object not registerd to stream
+	 */
+	IMPLEMENT_CUSTOM_SIMPLE_AUTOMATION_TEST(FPutAddAndRemoveFrequencyFlow, FChangeStreamsTestBase, "Editor.Concert.Replication.Stream.Frequency.PutAddAndRemoveFrequencyFlow", EAutomationTestFlags::EditorContext | EAutomationTestFlags::EngineFilter);
+	bool FPutAddAndRemoveFrequencyFlow::RunTest(const FString& Parameters)
+	{
+		UTestReflectionObject* MainTestObject = NewObject<UTestReflectionObject>(GetTransientPackage());
+		UTestReflectionObject* SecondaryTestObject = NewObject<UTestReflectionObject>(GetTransientPackage());
+		auto[FloatStreamId, FloatStream] = CreateFloatPropertyStream(*MainTestObject);
+		SenderArgs.Streams = { FloatStream };
+		SetUpClientAndServer();
+
+		FConcertBaseStreamInfo ExpectedStream = FloatStream.BaseDescription;
+		// Put SpecifiedRate 21
+		{
+			FConcertReplication_ChangeStream_Request ChangeRequest;
+			FConcertReplication_ChangeStream_Frequency FrequencyChange;
+			FrequencyChange.OverridesToPut.Add(MainTestObject, { EConcertObjectReplicationMode::SpecifiedRate, 21 });
+			ChangeRequest.FrequencyChanges.Add(FloatStreamId, FrequencyChange);
+			
+			ExpectedStream.FrequencySettings.ObjectOverrides.Add(MainTestObject, { EConcertObjectReplicationMode::SpecifiedRate, 21 });
+			ChangeStreamForSenderClientAndValidate(TEXT("Put SpecifiedRate 21"), ChangeRequest, { ExpectedStream });
+		}
+		// Add SpecifiedRate 22
+		{
+			FConcertReplication_ChangeStream_Request ChangeRequest;
+			FConcertReplication_ChangeStream_Frequency FrequencyChange;
+			FrequencyChange.OverridesToAdd.Add(MainTestObject, { EConcertObjectReplicationMode::SpecifiedRate, 22 });
+			ChangeRequest.FrequencyChanges.Add(FloatStreamId, FrequencyChange);
+			
+			ExpectedStream.FrequencySettings.ObjectOverrides.Add(MainTestObject, { EConcertObjectReplicationMode::SpecifiedRate, 22 });
+			ChangeStreamForSenderClientAndValidate(TEXT("Add SpecifiedRate 22"), ChangeRequest, { ExpectedStream });
+		}
+		// Remove 
+		{
+			FConcertReplication_ChangeStream_Request ChangeRequest;
+			FConcertReplication_ChangeStream_Frequency FrequencyChange;
+			FrequencyChange.OverridesToRemove.Add(MainTestObject);
+			ChangeRequest.FrequencyChanges.Add(FloatStreamId, FrequencyChange);
+			
+			ExpectedStream.FrequencySettings.ObjectOverrides.Remove(MainTestObject);
+			ChangeStreamForSenderClientAndValidate(TEXT("Remove"), ChangeRequest, { ExpectedStream });
+		}
+		
+		return true;
+	}
+
+	/**
+	 * Tests that requests for unregistered objects rejected:
+	 * - Cannot change frequency for unknown stream
+	 * - Cannot change frequency for object not registerd to stream
+	 */
+	IMPLEMENT_CUSTOM_SIMPLE_AUTOMATION_TEST(FRemovingObjectRemovesFrequencyEntry, FChangeStreamsTestBase, "Editor.Concert.Replication.Stream.Frequency.RemovingObjectRemovesFrequencyEntry", EAutomationTestFlags::EditorContext | EAutomationTestFlags::EngineFilter);
+	bool FRemovingObjectRemovesFrequencyEntry::RunTest(const FString& Parameters)
+	{
+		// 1. Start up client with float property stream
+		UTestReflectionObject* MainTestObject = NewObject<UTestReflectionObject>(GetTransientPackage());
+		UTestReflectionObject* SecondaryTestObject = NewObject<UTestReflectionObject>(GetTransientPackage());
+		auto[FloatStreamId, FloatStream] = CreateFloatPropertyStream(*MainTestObject);
+		SenderArgs.Streams = { FloatStream };
+		
+		SetUpClientAndServer();
+		FConcertBaseStreamInfo ExpectedStream = FloatStream.BaseDescription;
+		
+		// Register SecondaryObject with stream
+		{
+			FConcertReplication_ChangeStream_Request ChangeRequest;
+			
+			FConcertReplication_ChangeStream_PutObject PutObject;
+			AddFloatProperty(PutObject.Properties);
+			PutObject.ClassPath = SecondaryTestObject->GetClass();
+			ChangeRequest.ObjectsToPut.Add({ FloatStreamId, SecondaryTestObject }, PutObject);
+
+			ExpectedStream.ReplicationMap.ReplicatedObjects.Add(SecondaryTestObject, FConcertReplicatedObjectInfo{ PutObject.ClassPath, PutObject.Properties });
+			ChangeStreamForSenderClientAndValidate(TEXT("Add MainTestObject & SecondaryObject"), ChangeRequest, { ExpectedStream });
+		}
+		// Put frequency for MainTestObject & SecondaryObject
+		{
+			FConcertReplication_ChangeStream_Request ChangeRequest;
+			FConcertReplication_ChangeStream_Frequency FrequencyChange;
+			FrequencyChange.OverridesToPut.Add(MainTestObject, { EConcertObjectReplicationMode::SpecifiedRate, 21 });
+			FrequencyChange.OverridesToPut.Add(SecondaryTestObject, { EConcertObjectReplicationMode::SpecifiedRate, 22 });
+			ChangeRequest.FrequencyChanges.Add(FloatStreamId, FrequencyChange);
+			
+			ExpectedStream.FrequencySettings.ObjectOverrides.Add(MainTestObject, { EConcertObjectReplicationMode::SpecifiedRate, 21 });
+			ExpectedStream.FrequencySettings.ObjectOverrides.Add(SecondaryTestObject, { EConcertObjectReplicationMode::SpecifiedRate, 22 });
+			ChangeStreamForSenderClientAndValidate(TEXT("Put SpecifiedRate 21 & 22"), ChangeRequest, { ExpectedStream });
+		}
+		// Removing MainTestObject removes replication frequency entry
+		{
+			FConcertReplication_ChangeStream_Request ChangeRequest;
+			ChangeRequest.ObjectsToRemove.Add({ FloatStreamId, MainTestObject });
+
+			ExpectedStream.ReplicationMap.ReplicatedObjects.Remove(MainTestObject);
+			ExpectedStream.FrequencySettings.ObjectOverrides.Remove(MainTestObject);
+			ChangeStreamForSenderClientAndValidate(TEXT("Unregister object > removes frequency entry"), ChangeRequest, { ExpectedStream });
+		}
+
+		return true;
+	}
+
+	/**
+	 * Tests that requests for unregistered objects rejected:
+	 * - Cannot change frequency for unknown stream
+	 * - Cannot change frequency for object not registerd to stream
+	 */
+	IMPLEMENT_CUSTOM_SIMPLE_AUTOMATION_TEST(FPutReplacesAll, FChangeStreamsTestBase, "Editor.Concert.Replication.Stream.Frequency.PutReplacesAll", EAutomationTestFlags::EditorContext | EAutomationTestFlags::EngineFilter);
+	bool FPutReplacesAll::RunTest(const FString& Parameters)
+	{
+		// 1. Start up client with float property stream
+		UTestReflectionObject* MainTestObject = NewObject<UTestReflectionObject>(GetTransientPackage());
+		UTestReflectionObject* SecondaryTestObject = NewObject<UTestReflectionObject>(GetTransientPackage());
+		auto[FloatStreamId, FloatStream] = CreateFloatPropertyStream(*MainTestObject);
+		SenderArgs.Streams = { FloatStream };
+		SetUpClientAndServer();
+		
+		FConcertBaseStreamInfo ExpectedStream = FloatStream.BaseDescription;
+		// Register SecondaryObject with stream
+		{
+			FConcertReplication_ChangeStream_Request ChangeRequest;
+			
+			FConcertReplication_ChangeStream_PutObject PutObject;
+			AddFloatProperty(PutObject.Properties);
+			PutObject.ClassPath = SecondaryTestObject->GetClass();
+			ChangeRequest.ObjectsToPut.Add({ FloatStreamId, SecondaryTestObject }, PutObject);
+
+			ExpectedStream.ReplicationMap.ReplicatedObjects.Add(SecondaryTestObject, FConcertReplicatedObjectInfo{ PutObject.ClassPath, PutObject.Properties });
+			ChangeStreamForSenderClientAndValidate(TEXT("Add MainTestObject & SecondaryObject"), ChangeRequest, { ExpectedStream });
+		}
+		
+
+		// Put MainTestObject & SecondaryObject
+		{
+			FConcertReplication_ChangeStream_Request ChangeRequest;
+			FConcertReplication_ChangeStream_Frequency FrequencyChange;
+			FrequencyChange.OverridesToPut.Add(MainTestObject, { EConcertObjectReplicationMode::SpecifiedRate, 10 });
+			FrequencyChange.OverridesToPut.Add(SecondaryTestObject, { EConcertObjectReplicationMode::SpecifiedRate, 20 });
+			ChangeRequest.FrequencyChanges.Add(FloatStreamId, FrequencyChange);
+			
+			ExpectedStream.FrequencySettings.ObjectOverrides.Add(MainTestObject, { EConcertObjectReplicationMode::SpecifiedRate, 10 });
+			ExpectedStream.FrequencySettings.ObjectOverrides.Add(SecondaryTestObject, { EConcertObjectReplicationMode::SpecifiedRate, 20 });
+			ChangeStreamForSenderClientAndValidate(TEXT("Add MainTestObject & SecondaryObject"), ChangeRequest, { ExpectedStream });
+		}
+		
+		// Put replaces all
+		{
+			FConcertReplication_ChangeStream_Request ChangeRequest;
+			FConcertReplication_ChangeStream_Frequency FrequencyChange;
+			FrequencyChange.OverridesToPut.Add(MainTestObject, { EConcertObjectReplicationMode::SpecifiedRate, 100 });
+			ChangeRequest.FrequencyChanges.Add(FloatStreamId, FrequencyChange);
+
+			ExpectedStream.FrequencySettings.ObjectOverrides.Add(MainTestObject, { EConcertObjectReplicationMode::SpecifiedRate, 100 });
+			ExpectedStream.FrequencySettings.ObjectOverrides.Remove(SecondaryTestObject);
+			ChangeStreamForSenderClientAndValidate(TEXT("Put replaces all overrides"), ChangeRequest, { ExpectedStream });
+		}
+		
+		return true;
+	}
+
+	/**
+	 * Tests that requests for unregistered objects rejected:
+	 * - Cannot change frequency for unknown stream
+	 * - Cannot change frequency for object not registerd to stream
+	 */
+	IMPLEMENT_CUSTOM_SIMPLE_AUTOMATION_TEST(FAddAfterPut, FChangeStreamsTestBase, "Editor.Concert.Replication.Stream.Frequency.AddAfterPut", EAutomationTestFlags::EditorContext | EAutomationTestFlags::EngineFilter);
+	bool FAddAfterPut::RunTest(const FString& Parameters)
+	{
+		// 1. Start up client with float property stream
+		UTestReflectionObject* MainTestObject = NewObject<UTestReflectionObject>(GetTransientPackage());
+		auto[FloatStreamId, FloatStream] = CreateFloatPropertyStream(*MainTestObject);
+		SenderArgs.Streams = { FloatStream };
+		SetUpClientAndServer();
+
+		FConcertBaseStreamInfo ExpectedStream = FloatStream.BaseDescription;
+		// Put SpecifiedRate 25, then add 50
+		{
+			FConcertReplication_ChangeStream_Request ChangeRequest;
+			FConcertReplication_ChangeStream_Frequency FrequencyChange;
+			FrequencyChange.OverridesToPut.Add(MainTestObject, { EConcertObjectReplicationMode::SpecifiedRate, 25 });
+			FrequencyChange.OverridesToAdd.Add(MainTestObject, { EConcertObjectReplicationMode::SpecifiedRate, 50 });
+			ChangeRequest.FrequencyChanges.Add(FloatStreamId, FrequencyChange);
+			
+			ExpectedStream.FrequencySettings.ObjectOverrides.Add(MainTestObject, { EConcertObjectReplicationMode::SpecifiedRate, 50 });
+			ChangeStreamForSenderClientAndValidate(TEXT("Put SpecifiedRate 25, then add 50"), ChangeRequest, { ExpectedStream });
+		}
+		
+		return true;
+	}
+
+	/**
+	 * Tests that requests for unregistered objects rejected:
+	 * - Cannot change frequency for unknown stream
+	 * - Cannot change frequency for object not registerd to stream
+	 */
+	IMPLEMENT_CUSTOM_SIMPLE_AUTOMATION_TEST(FCannotChangeFrequencyOfUnregistered, FChangeStreamsTestBase, "Editor.Concert.Replication.Stream.Frequency.FailWithUnregisteredStreamOrObject", EAutomationTestFlags::EditorContext | EAutomationTestFlags::EngineFilter);
+	bool FCannotChangeFrequencyOfUnregistered::RunTest(const FString& Parameters)
+	{
+		// 1. Start up client with float property stream
+		UTestReflectionObject* MainTestObject = NewObject<UTestReflectionObject>(GetTransientPackage());
+		UTestReflectionObject* SecondaryTestObject = NewObject<UTestReflectionObject>(GetTransientPackage());
+		FGuid FloatStreamId;
+		FConcertReplicationStream FloatStream;
+		Tie(FloatStreamId, FloatStream) = CreateFloatPropertyStream(*MainTestObject);
+		SenderArgs.Streams = { FloatStream };
+		
+		SetUpClientAndServer();
+		// We'll intentionally be causing errors below that cause a warning to be logged.
+		AddExpectedError(TEXT("Rejecting ChangeStream"),EAutomationExpectedErrorFlags::Contains, 2);
+
+		// Change defaults to be SpecifiedRate at 42 update rate.
+		FConcertBaseStreamInfo ExpectedStream = FloatStream.BaseDescription;
+		
+		// Cannot change frequency of unregistered stream
+		{
+			const FGuid UnregisteredStreamId = FGuid::NewGuid();
+			
+			FConcertReplication_ChangeStream_Request ChangeRequest;
+			FConcertReplication_ChangeStream_Frequency FrequencyChange;
+			FrequencyChange.OverridesToPut.Add(MainTestObject, { EConcertObjectReplicationMode::SpecifiedRate, 50 });
+			ChangeRequest.FrequencyChanges.Add(UnregisteredStreamId, FrequencyChange);
+
+			ChangeStreamForSenderClientAndValidate(
+				TEXT("Cannot change frequency of unregistered stream"),
+				ChangeRequest,
+				{ ExpectedStream },
+				EChangeStreamValidation::TestWasHandled | EChangeStreamValidation::TestResponseFailed
+				)
+				.Next([this, &UnregisteredStreamId, &MainTestObject](FConcertReplication_ChangeStream_Response&& Response)
+				{
+					const EConcertChangeObjectFrequencyErrorCode* ErrorCode = Response.FrequencyErrors.OverrideFailures.Find({ UnregisteredStreamId, MainTestObject });
+					TestTrue(TEXT("Cannot change frequency of unregistered stream > Rejection code (defaults)"), Response.FrequencyErrors.DefaultFailures.IsEmpty());
+					TestTrue(TEXT("Cannot change frequency of unregistered stream > Rejection code (object)"), ErrorCode && *ErrorCode == EConcertChangeObjectFrequencyErrorCode::NotRegistered);
+				});
+		}
+
+		// Cannot add object that is not registered
+		{
+			FConcertReplication_ChangeStream_Request ChangeRequest;
+			FConcertReplication_ChangeStream_Frequency FrequencyChange;
+			FrequencyChange.OverridesToPut.Add(SecondaryTestObject, { EConcertObjectReplicationMode::SpecifiedRate, 100 });
+			ChangeRequest.FrequencyChanges.Add(FloatStreamId, FrequencyChange);
+			
+			ChangeStreamForSenderClientAndValidate(
+				TEXT("Cannot add non-registered object"),
+				ChangeRequest,
+				{ ExpectedStream },
+				EChangeStreamValidation::TestWasHandled | EChangeStreamValidation::TestResponseFailed
+				)
+				.Next([this, FloatStreamId, SecondaryTestObject](FConcertReplication_ChangeStream_Response&& Response)
+				{
+					const EConcertChangeObjectFrequencyErrorCode* ErrorCode = Response.FrequencyErrors.OverrideFailures.Find({ FloatStreamId, SecondaryTestObject });
+					TestTrue(TEXT("Cannot add non-registered object > Rejection code (defaults)"), Response.FrequencyErrors.DefaultFailures.IsEmpty());
+					TestTrue(TEXT("Cannot add non-registered object > Rejection code (object)"), ErrorCode && *ErrorCode == EConcertChangeObjectFrequencyErrorCode::NotRegistered);
+				});
+		}
+
+		return true;
+	}
+}
\ No newline at end of file
diff --git a/Engine/Plugins/Developer/Concert/ConcertSync/ConcertSyncTest/Source/ConcertSyncTest/Private/Replication/Stream/StreamRequestTests_Query.cpp b/Engine/Plugins/Developer/Concert/ConcertSync/ConcertSyncTest/Source/ConcertSyncTest/Private/Replication/Stream/StreamRequestTests_Query.cpp
new file mode 100644
--- /dev/null
+++ b/Engine/Plugins/Developer/Concert/ConcertSync/ConcertSyncTest/Source/ConcertSyncTest/Private/Replication/Stream/StreamRequestTests_Query.cpp
@@ -0,0 +1,198 @@
+// Copyright Epic Games, Inc. All Rights Reserved.
+
+#include "Util/ClientServerCommunicationTest.h"
+
+#include "Replication/Data/ReplicationStream.h"
+#include "Replication/IConcertClientReplicationManager.h"
+#include "Replication/Util/ChangeStreamsTestBase.h"
+#include "Replication/Util/SendReceiveObjectTestBase.h"
+#include "Replication/TestReflectionObject.h"
+
+#include "Misc/AutomationTest.h"
+#include "UObject/UObjectGlobals.h"
+#include "UObject/Package.h"
+
+namespace UE::ConcertSyncTests::Replication::Stream
+{
+	/**
+	 * Let's receiving client query the server what the sending client is sending. Checks
+	 * - the reported streams,
+	 * - that the authority updates,
+	 * - the EConcertQueryClientStreamFlags flags
+	 *
+	 * @note The other tests in StreamRequestTests_{x}.cpp use QueryStream under the hood so all this test does is checking that the skip flags work correctly. 
+	 */
+	IMPLEMENT_CUSTOM_SIMPLE_AUTOMATION_TEST(FQueryOtherClientStreams, FSendReceiveObjectTestBase, "Editor.Concert.Replication.Stream.QueryClientStream", EAutomationTestFlags::EditorContext | EAutomationTestFlags::EngineFilter);
+	bool FQueryOtherClientStreams::RunTest(const FString& Parameters)
+	{
+		// 1. Init
+		SetUpClientAndServer();
+		
+		const FGuid SenderEndpointId = Client_Sender->ClientSessionMock->GetSessionClientEndpointId();
+		const FSoftObjectPath TestObjectPath{ TestObject };
+		FConcertReplication_QueryReplicationInfo_Request Request;
+		Request.ClientEndpointIds = { SenderEndpointId };
+		auto TestReplicationMapContent = [this, &SenderEndpointId](const FConcertReplication_QueryReplicationInfo_Response& Response)
+		{
+			TestTrue(TEXT("ErrorCode == Handled"), Response.ErrorCode == EReplicationResponseErrorCode::Handled);
+			
+			const FConcertQueriedClientInfo* Info = Response.ClientInfo.Find(SenderEndpointId);
+			if (!Info)
+			{
+				AddError(TEXT("No info about sending client!"));
+				return;
+			}
+				
+			TestEqual(TEXT("Only received info about requested client"), Response.ClientInfo.Num(), 1);
+			TestEqual(TEXT("Exactly one stream data"), Info->Streams.Num(), 1);
+			if (Info->Streams.IsEmpty())
+			{
+				return;
+			}
+				
+			const FConcertBaseStreamInfo& StreamDescription = Info->Streams[0];
+			TestEqual(TEXT("StreamDescription->Identifier correct"), StreamDescription.Identifier, SenderStreamId);
+			TestEqual(TEXT("StreamDescription->ReplicationMap has exactly 1 object"), StreamDescription.ReplicationMap.ReplicatedObjects.Num(), 1);
+
+			const FConcertObjectReplicationMap ExpectedReplicationMap = CreateSenderArgs().Streams[0].BaseDescription.ReplicationMap;
+			TestEqual(TEXT("Registered and reported replication maps match"), StreamDescription.ReplicationMap, ExpectedReplicationMap);
+		};
+		
+		// 2.1 Request before taking authority
+		bool bReceivedFirstQueryResponse = false;
+		ClientReplicationManager_Receiver->QueryClientInfo(Request)
+			.Next([this, &SenderEndpointId, &TestReplicationMapContent, &bReceivedFirstQueryResponse](FConcertReplication_QueryReplicationInfo_Response&& Response) mutable
+			{
+				bReceivedFirstQueryResponse = true;
+				TestTrue(TEXT("ErrorCode == Handled"), Response.ErrorCode == EReplicationResponseErrorCode::Handled);
+				
+				TestReplicationMapContent(Response);
+				if (const FConcertQueriedClientInfo* Info = Response.ClientInfo.Find(SenderEndpointId))
+				{
+					TestEqual(TEXT("Contains no authority data"), Info->Authority.Num(), 0);
+				}
+			});
+		TestTrue(TEXT("Received query response (before taking authority)"), bReceivedFirstQueryResponse);
+
+		// 2.2 Request after taking authority
+		bool bSenderReceivedResponse = false;
+		ClientReplicationManager_Sender->TakeAuthorityOver({ TestObjectPath })
+			.Next([this, &bSenderReceivedResponse](const FConcertReplication_ChangeAuthority_Response& Response) mutable
+			{
+				bSenderReceivedResponse = true;
+				TestEqual(TEXT("No rejection taking authority"), Response.RejectedObjects.Num(), 0); 
+				TestTrue(TEXT("ErrorCode == Handled"), Response.ErrorCode == EReplicationResponseErrorCode::Handled);
+			});
+		TestTrue(TEXT("Received response taking authority"), bSenderReceivedResponse);
+
+		bool bReceivedSecondQueryResponse = false;
+		ClientReplicationManager_Receiver->QueryClientInfo(Request)
+			.Next([this, &SenderEndpointId, &TestReplicationMapContent, &TestObjectPath, &bReceivedSecondQueryResponse](FConcertReplication_QueryReplicationInfo_Response&& Response) mutable
+			{
+				bReceivedSecondQueryResponse = true;
+				TestTrue(TEXT("ErrorCode == Handled"), Response.ErrorCode == EReplicationResponseErrorCode::Handled);
+				TestReplicationMapContent(Response);
+				
+				const FConcertQueriedClientInfo* Info = Response.ClientInfo.Find(SenderEndpointId);
+				if (!Info || Info->Authority.IsEmpty())
+				{
+					AddError(TEXT("Missing authority data"));
+					return;
+				}
+
+				TestEqual(TEXT("Exactly 1 authority stream"), Info->Authority.Num(), 1);
+				const FConcertAuthorityClientInfo& AuthorityInfo = Info->Authority[0];
+				TestEqual(TEXT("Authority stream ID matches registered stream ID"), AuthorityInfo.StreamId, SenderStreamId);
+				TestEqual(TEXT("Has authority over exactly 1 object"), AuthorityInfo.AuthoredObjects.Num(), 1);
+				TestTrue(TEXT("Has authority over registered object"), AuthorityInfo.AuthoredObjects.Contains(TestObjectPath));
+			});
+		TestTrue(TEXT("Received response to taking authority"), bSenderReceivedResponse);
+		TestTrue(TEXT("Received query response (after taking authority)"), bReceivedSecondQueryResponse);
+
+		// 2.3 SkipStreamInfo
+		bool bReceivedResponse_SkipStreamInfo = false;
+		Request.QueryFlags = EConcertQueryClientStreamFlags::SkipStreamInfo;
+		ClientReplicationManager_Receiver->QueryClientInfo(Request)
+			.Next([this, &SenderEndpointId, &bReceivedResponse_SkipStreamInfo](FConcertReplication_QueryReplicationInfo_Response&& Response) mutable
+			{
+				bReceivedResponse_SkipStreamInfo = true;
+				TestTrue(TEXT("ErrorCode == Handled"), Response.ErrorCode == EReplicationResponseErrorCode::Handled);
+				if (const FConcertQueriedClientInfo* Info = Response.ClientInfo.Find(SenderEndpointId))
+				{
+					TestEqual(TEXT("SkipStreamInfo > No stream data"), Info->Streams.Num(), 0);
+				}
+			});
+		TestTrue(TEXT("bReceivedResponse_SkipStreamInfo"), bReceivedResponse_SkipStreamInfo);
+		
+		// 2.4 SkipProperties
+		bool bReceivedResponse_SkipProperties = false;
+		Request.QueryFlags = EConcertQueryClientStreamFlags::SkipProperties;
+		ClientReplicationManager_Receiver->QueryClientInfo(Request)
+			.Next([this, &SenderEndpointId, &TestObjectPath, &bReceivedResponse_SkipProperties](FConcertReplication_QueryReplicationInfo_Response&& Response) mutable
+			{
+				bReceivedResponse_SkipProperties = true;
+				const FConcertQueriedClientInfo* Info = Response.ClientInfo.Find(SenderEndpointId);
+				TestTrue(TEXT("ErrorCode == Handled"), Response.ErrorCode == EReplicationResponseErrorCode::Handled);
+				if (!Info || Info->Streams.IsEmpty())
+				{
+					AddError(TEXT("SkipProperties > No stream info"));
+					return;
+				}
+
+				const FConcertReplicatedObjectInfo* ObjectInfo = Info->Streams[0].ReplicationMap.ReplicatedObjects.Find(TestObjectPath);
+				if (!ObjectInfo)
+				{
+					AddError(TEXT("SkipProperties > No object info"));
+					return;
+				}
+				TestEqual(TEXT("SkipProperties > No property data"), ObjectInfo->PropertySelection.ReplicatedProperties.Num(), 0);
+			});
+		TestTrue(TEXT("bReceivedResponse_SkipProperties"), bReceivedResponse_SkipProperties);
+		
+		// 2.5 SkipAuthority
+		bool bReceivedResponse_SkipAuthority = false;
+		Request.QueryFlags = EConcertQueryClientStreamFlags::SkipAuthority;
+		ClientReplicationManager_Receiver->QueryClientInfo(Request)
+			.Next([this, &SenderEndpointId, &bReceivedResponse_SkipAuthority](FConcertReplication_QueryReplicationInfo_Response&& Response) mutable
+			{
+				bReceivedResponse_SkipAuthority = true;
+				TestTrue(TEXT("ErrorCode == Handled"), Response.ErrorCode == EReplicationResponseErrorCode::Handled);
+				if (const FConcertQueriedClientInfo* Info = Response.ClientInfo.Find(SenderEndpointId))
+				{
+					TestEqual(TEXT("SkipAuthority > No authority data"), Info->Authority.Num(), 0);
+				}
+			});
+		TestTrue(TEXT("bReceivedResponse_SkipAuthority"), bReceivedResponse_SkipAuthority);
+
+		// 2.6 SkipFrequency
+		FConcertReplication_ChangeStream_Request AddFrequencyInfo;
+		FConcertReplication_ChangeStream_Frequency& FrequencyChange = AddFrequencyInfo.FrequencyChanges.Add(SenderStreamId);
+		FrequencyChange.OverridesToPut.Add(TestObjectPath) = { EConcertObjectReplicationMode::SpecifiedRate, 20 }; 
+		FrequencyChange.NewDefaults = { EConcertObjectReplicationMode::SpecifiedRate, 42 };
+		FrequencyChange.Flags = EConcertReplicationChangeFrequencyFlags::SetDefaults;
+		// Important: Change Sender here, not Receiver! Receiver is the one querying for Sender.
+		ClientReplicationManager_Sender->ChangeStream(AddFrequencyInfo)
+			.Next([this](FConcertReplication_ChangeStream_Response&& Response)
+			{
+				TestTrue(TEXT("Changed frequency"), Response.IsSuccess());
+			});
+		
+		bool bReceivedResponse_SkipFrequency = false;
+		Request.QueryFlags = EConcertQueryClientStreamFlags::SkipFrequency;
+		ClientReplicationManager_Receiver->QueryClientInfo(Request)
+			.Next([this, &SenderEndpointId, &bReceivedResponse_SkipFrequency](FConcertReplication_QueryReplicationInfo_Response&& Response) mutable
+			{
+				bReceivedResponse_SkipFrequency = true;
+				TestTrue(TEXT("ErrorCode == Handled"), Response.ErrorCode == EReplicationResponseErrorCode::Handled);
+				if (const FConcertQueriedClientInfo* Info = Response.ClientInfo.Find(SenderEndpointId)
+					; Info && ensureAlways(Info->Streams.Num() == 1))
+				{
+					const FConcertStreamFrequencySettings& FrequencySettings = Info->Streams[0].FrequencySettings;
+					TestEqual(TEXT("SkipFrequency > No frequency override data"), FrequencySettings.ObjectOverrides.Num(), 0);
+				}
+			});
+		TestTrue(TEXT("bReceivedResponse_SkipFrequency"), bReceivedResponse_SkipFrequency);
+		
+		return true;
+	}
+}
\ No newline at end of file
diff --git a/Engine/Plugins/Developer/Concert/ConcertSync/ConcertSyncTest/Source/ConcertSyncTest/Private/Replication/StreamRequestTests_Change.cpp b/Engine/Plugins/Developer/Concert/ConcertSync/ConcertSyncTest/Source/ConcertSyncTest/Private/Replication/StreamRequestTests_Change.cpp
deleted file mode 100644
diff --git a/Engine/Plugins/Developer/Concert/ConcertSync/ConcertSyncTest/Source/ConcertSyncTest/Private/Replication/StreamRequestTests_Frequency.cpp b/Engine/Plugins/Developer/Concert/ConcertSync/ConcertSyncTest/Source/ConcertSyncTest/Private/Replication/StreamRequestTests_Frequency.cpp
deleted file mode 100644
diff --git a/Engine/Plugins/Developer/Concert/ConcertSync/ConcertSyncTest/Source/ConcertSyncTest/Private/Replication/StreamRequestTests_Query.cpp b/Engine/Plugins/Developer/Concert/ConcertSync/ConcertSyncTest/Source/ConcertSyncTest/Private/Replication/StreamRequestTests_Query.cpp
deleted file mode 100644
