Fix FConcertReplication_ChangeStream_Request being rejected when it specifies the same stream ID in StreamsToAdd and StreamsToRemove.
The server now will replace the stream content as intended.

diff --git a/Engine/Plugins/Developer/Concert/ConcertSync/ConcertSyncServer/Source/ConcertSyncServer/Private/Replication/ConcertServerReplicationManager_ChangeStream.cpp b/Engine/Plugins/Developer/Concert/ConcertSync/ConcertSyncServer/Source/ConcertSyncServer/Private/Replication/ConcertServerReplicationManager_ChangeStream.cpp
--- a/Engine/Plugins/Developer/Concert/ConcertSync/ConcertSyncServer/Source/ConcertSyncServer/Private/Replication/ConcertServerReplicationManager_ChangeStream.cpp
+++ b/Engine/Plugins/Developer/Concert/ConcertSync/ConcertSyncServer/Source/ConcertSyncServer/Private/Replication/ConcertServerReplicationManager_ChangeStream.cpp
@@ -73,12 +73,14 @@ namespace UE::ConcertSyncServer::Replication
 				// StreamsToAdd is invalid if there is already a stream with the same ID registered ...
 				const FGuid& NewStreamId = NewStream.BaseDescription.Identifier;
 				const bool bIdAlreadyExists = FindExistingStream(Client, NewStreamId) != nullptr;
+				const bool bIsStreamRemoved = Request.StreamsToRemove.Contains(NewStreamId);
 
 				// ... or StreamsToAdd contains the same ID multiple times
 				bool bIsDuplicateEntry = false;
 				DuplicateEntryDetection.FindOrAdd(NewStreamId, &bIsDuplicateEntry);
 				
-				if (bIdAlreadyExists || bIsDuplicateEntry)
+				if ((bIdAlreadyExists && !bIsStreamRemoved)
+					|| bIsDuplicateEntry)
 				{
 					UE_LOG(LogConcert, Log, TEXT("Duplicate stream entry %s"), *NewStreamId.ToString(EGuidFormats::Short));
 					OutResponse.FailedStreamCreation.Add(NewStreamId);
diff --git a/Engine/Plugins/Developer/Concert/ConcertSync/ConcertSyncTest/Source/ConcertSyncTest/Private/Replication/Stream/ChangeStream.spec.cpp b/Engine/Plugins/Developer/Concert/ConcertSync/ConcertSyncTest/Source/ConcertSyncTest/Private/Replication/Stream/ChangeStream.spec.cpp
new file mode 100644
--- /dev/null
+++ b/Engine/Plugins/Developer/Concert/ConcertSync/ConcertSyncTest/Source/ConcertSyncTest/Private/Replication/Stream/ChangeStream.spec.cpp
@@ -0,0 +1,83 @@
+﻿// Copyright Epic Games, Inc. All Rights Reserved.
+
+#include "Replication/Util/Mocks/ReplicationWorkspaceCallInterceptorMock.h"
+#include "Replication/Util/Spec/ReplicationClient.h"
+#include "Replication/Util/Spec/ReplicationServer.h"
+#include "Replication/Util/Spec/ObjectTestReplicator.h"
+
+#include "Misc/AutomationTest.h"
+#include "Replication/ReplicationTestInterface.h"
+#include "Util/ScopedSessionDatabase.h"
+
+namespace UE::ConcertSyncTests::Replication::RestoreContent
+{
+	BEGIN_DEFINE_SPEC(FChangeStreamSpec, "Editor.Concert.Replication.Stream", EAutomationTestFlags::EditorContext | EAutomationTestFlags::EngineFilter)
+		TUniquePtr<FObjectTestReplicator> FooObject;
+		TUniquePtr<FObjectTestReplicator> BarObject;
+	
+		TUniquePtr<FReplicationServer> Server;
+		FReplicationClient* Client = nullptr;
+	
+		const FGuid StreamId = FGuid::NewGuid();
+	END_DEFINE_SPEC(FChangeStreamSpec);
+
+	/**
+	 * This tests use cases that only issue FConcertReplication_ChangeStream_Request.
+	 * 
+	 * Old tests are in StreamRequestTests_x.cpp files of the same folder.
+	 * In the future, we want to move those over to the spec test format, too.
+	 */
+	void FChangeStreamSpec::Define()
+	{
+		BeforeEach([this]
+		{
+			FooObject = MakeUnique<FObjectTestReplicator>(TEXT("Foo"));
+			BarObject = MakeUnique<FObjectTestReplicator>(TEXT("Bar"));
+			
+			Server = MakeUnique<FReplicationServer>(*this);
+			Client = &Server->ConnectClient();
+			Client->JoinReplication();
+		});
+		AfterEach([this]
+		{
+			Server.Reset();
+			BarObject.Reset();
+			FooObject.Reset();
+		});
+
+		It("Replace stream content", [this]
+		{
+			IConcertClientReplicationManager& ReplicationManager = Client->GetClientReplicationManager();
+			ReplicationManager.ChangeStream({ .StreamsToAdd = { FooObject->CreateStream(StreamId) } });
+
+			const FConcertReplicationStream BarStream = BarObject->CreateStream(StreamId);
+			const FConcertReplication_ChangeStream_Request ReplaceWithBarRequest
+			{
+				.StreamsToAdd = { BarStream },
+				.StreamsToRemove = { StreamId  }
+			};
+			bool bReceivedStreamResponse = false;
+			ReplicationManager
+				.ChangeStream(ReplaceWithBarRequest)
+				.Next([this, &bReceivedStreamResponse](FConcertReplication_ChangeStream_Response&& Response)
+				{
+					bReceivedStreamResponse = true;
+					TestTrue(TEXT("Success"), Response.IsSuccess());
+				});
+			TestTrue(TEXT("bReceivedResponse"), bReceivedStreamResponse);
+
+			bool bReceivedQueryResponse = false;
+			ReplicationManager
+				.QueryClientInfo({ .ClientEndpointIds = { Client->GetEndpointId() } })
+				.Next([this, &bReceivedQueryResponse, &BarStream](FConcertReplication_QueryReplicationInfo_Response&& Response)
+				{
+					bReceivedQueryResponse = true;
+					
+					const FConcertQueriedClientInfo* ClientInfo = Response.ClientInfo.Find(Client->GetEndpointId());
+					TestTrue(TEXT("Client info"), ClientInfo && ClientInfo->Streams.Num() == 1 && ClientInfo->Streams[0] == BarStream.BaseDescription);
+					
+				});
+			TestTrue(TEXT("bReceivedQueryResponse"), bReceivedStreamResponse);
+		});
+	}
+}
diff --git a/Engine/Plugins/Developer/Concert/ConcertSync/ConcertSyncTest/Source/ConcertSyncTest/Private/Replication/Stream/StreamAuthorityInterplay.spec.cpp b/Engine/Plugins/Developer/Concert/ConcertSync/ConcertSyncTest/Source/ConcertSyncTest/Private/Replication/Stream/StreamAuthorityInterplay.spec.cpp
--- a/Engine/Plugins/Developer/Concert/ConcertSync/ConcertSyncTest/Source/ConcertSyncTest/Private/Replication/Stream/StreamAuthorityInterplay.spec.cpp
+++ b/Engine/Plugins/Developer/Concert/ConcertSync/ConcertSyncTest/Source/ConcertSyncTest/Private/Replication/Stream/StreamAuthorityInterplay.spec.cpp
@@ -11,14 +11,14 @@
 
 namespace UE::ConcertSyncTests::Replication::RestoreContent
 {
-	BEGIN_DEFINE_SPEC(FChangeStreamSpec, "Editor.Concert.Replication.Stream", EAutomationTestFlags::EditorContext | EAutomationTestFlags::EngineFilter)
+	BEGIN_DEFINE_SPEC(FStreamAuthorityInterplaySpec, "Editor.Concert.Replication.Stream", EAutomationTestFlags::EditorContext | EAutomationTestFlags::EngineFilter)
 		TUniquePtr<FObjectTestReplicator> ObjectReplicator;
 	
 		TUniquePtr<FReplicationServer> Server;
 		FReplicationClient* Client = nullptr;
 
 		const FGuid StreamId = FGuid::NewGuid();
-	END_DEFINE_SPEC(FChangeStreamSpec);
+	END_DEFINE_SPEC(FStreamAuthorityInterplaySpec);
 
 	/**
 	 * This tests that stream structure and authority correctly interact when issuing a FConcertReplication_ChangeStream_Request.
@@ -26,7 +26,7 @@ namespace UE::ConcertSyncTests::Replication::RestoreContent
 	 * Old tests are in StreamRequestTests_x.cpp files of the same folder.
 	 * In the future, we want to move those over to the spec test format, too.
 	 */
-	void FChangeStreamSpec::Define()
+	void FStreamAuthorityInterplaySpec::Define()
 	{
 		BeforeEach([this]
 		{
