Fix the server not rejecting put state request that creates an empty stream.
This would break the invariant that streams on the server are non-empty.

diff --git a/Engine/Plugins/Developer/Concert/ConcertSync/ConcertSyncCore/Source/ConcertSyncCore/Private/Replication/ChangeStreamSharedUtils.cpp b/Engine/Plugins/Developer/Concert/ConcertSync/ConcertSyncCore/Source/ConcertSyncCore/Private/Replication/ChangeStreamSharedUtils.cpp
--- a/Engine/Plugins/Developer/Concert/ConcertSync/ConcertSyncCore/Source/ConcertSyncCore/Private/Replication/ChangeStreamSharedUtils.cpp
+++ b/Engine/Plugins/Developer/Concert/ConcertSync/ConcertSyncCore/Source/ConcertSyncCore/Private/Replication/ChangeStreamSharedUtils.cpp
@@ -134,6 +134,7 @@ namespace UE::ConcertSyncCore::Replication::ChangeStreamUtils
 
 		for (const FConcertReplicationStream& StreamDescription : Request.StreamsToAdd)
 		{
+			ensureAlwaysMsgf(!StreamDescription.BaseDescription.ReplicationMap.IsEmpty(), TEXT("This should have been validated"));
 			StreamsToModify.Add(StreamDescription);
 		}
 	}
diff --git a/Engine/Plugins/Developer/Concert/ConcertSync/ConcertSyncCore/Source/ConcertSyncCore/Public/Replication/Messages/ChangeStream.h b/Engine/Plugins/Developer/Concert/ConcertSync/ConcertSyncCore/Source/ConcertSyncCore/Public/Replication/Messages/ChangeStream.h
--- a/Engine/Plugins/Developer/Concert/ConcertSync/ConcertSyncCore/Source/ConcertSyncCore/Public/Replication/Messages/ChangeStream.h
+++ b/Engine/Plugins/Developer/Concert/ConcertSync/ConcertSyncCore/Source/ConcertSyncCore/Public/Replication/Messages/ChangeStream.h
@@ -294,7 +294,7 @@ struct FConcertReplication_ChangeStream_Response
 	UPROPERTY()
 	TMap<FConcertObjectInStreamID, EConcertPutObjectErrorCode> ObjectsToPutSemanticErrors;
 	
-	/** Streams that were in StreamsToAdd but that were not created (stream was already registered). */
+	/** Streams that were in StreamsToAdd but that were not created (stream was already registered, or was empty). */
 	UPROPERTY()
 	TSet<FGuid> FailedStreamCreation;
 
diff --git a/Engine/Plugins/Developer/Concert/ConcertSync/ConcertSyncServer/Source/ConcertSyncServer/Private/Replication/Util/StreamChangeValidation.cpp b/Engine/Plugins/Developer/Concert/ConcertSync/ConcertSyncServer/Source/ConcertSyncServer/Private/Replication/Util/StreamChangeValidation.cpp
--- a/Engine/Plugins/Developer/Concert/ConcertSync/ConcertSyncServer/Source/ConcertSyncServer/Private/Replication/Util/StreamChangeValidation.cpp
+++ b/Engine/Plugins/Developer/Concert/ConcertSync/ConcertSyncServer/Source/ConcertSyncServer/Private/Replication/Util/StreamChangeValidation.cpp
@@ -63,7 +63,7 @@ namespace UE::ConcertSyncServer::Replication
 		}
 
 		/** Checks that StreamsToAdd do not conflict with pre-existing ones and that all IDs in the request are also unique. */
-		static void ValidateAddedStreamsAreUnique(const FConcertReplication_ChangeStream_Request& Request, const FConcertReplicationClient& Client, FConcertReplication_ChangeStream_Response& OutResponse)
+		static void ValidateAddedStreamsAreValid(const FConcertReplication_ChangeStream_Request& Request, const FConcertReplicationClient& Client, FConcertReplication_ChangeStream_Response& OutResponse)
 		{
 			TSet<FGuid> DuplicateEntryDetection;
 			for (const FConcertReplicationStream& NewStream : Request.StreamsToAdd)
@@ -78,9 +78,10 @@ namespace UE::ConcertSyncServer::Replication
 				DuplicateEntryDetection.FindOrAdd(NewStreamId, &bIsDuplicateEntry);
 				
 				if ((bIdAlreadyExists && !bIsStreamRemoved)
-					|| bIsDuplicateEntry)
+					|| bIsDuplicateEntry
+					|| NewStream.BaseDescription.ReplicationMap.IsEmpty())
 				{
-					UE_LOG(LogConcert, Log, TEXT("Duplicate stream entry %s"), *NewStreamId.ToString(EGuidFormats::Short));
+					UE_LOG(LogConcert, Log, TEXT("Failed to create stream %s"), *NewStreamId.ToString(EGuidFormats::Short));
 					OutResponse.FailedStreamCreation.Add(NewStreamId);
 				}
 			}
@@ -131,7 +132,7 @@ namespace UE::ConcertSyncServer::Replication
 		OutResponse.ErrorCode = EReplicationResponseErrorCode::Handled;
 
 		Private::ValidatePutObjectsRequestSemantics(Request, Client, OutResponse);
-		Private::ValidateAddedStreamsAreUnique(Request, Client, OutResponse);
+		Private::ValidateAddedStreamsAreValid(Request, Client, OutResponse);
 		Private::LookForAuthorityConflicts(Request, Client, AuthorityManager, OutResponse);
 		ConcertSyncCore::Replication::ChangeStreamUtils::ValidateFrequencyChanges(Request, Client.GetStreamDescriptions(), &OutResponse.FrequencyErrors);
 			
diff --git a/Engine/Plugins/Developer/Concert/ConcertSync/ConcertSyncTest/Source/ConcertSyncTest/Private/Replication/PutState/PutStateMisc.spec.cpp b/Engine/Plugins/Developer/Concert/ConcertSync/ConcertSyncTest/Source/ConcertSyncTest/Private/Replication/PutState/PutStateMisc.spec.cpp
new file mode 100644
--- /dev/null
+++ b/Engine/Plugins/Developer/Concert/ConcertSync/ConcertSyncTest/Source/ConcertSyncTest/Private/Replication/PutState/PutStateMisc.spec.cpp
@@ -0,0 +1,77 @@
+﻿// Copyright Epic Games, Inc. All Rights Reserved.
+
+#include "Replication/Util/Spec/ReplicationClient.h"
+#include "Replication/Util/Spec/ReplicationServer.h"
+#include "Replication/Util/Spec/ObjectTestReplicator.h"
+
+#include "Misc/AutomationTest.h"
+#include "Replication/Messages/PutState.h"
+
+namespace UE::ConcertSyncTests::Replication::ChangeClients
+{
+	BEGIN_DEFINE_SPEC(FPutStateMiscSpec, "Editor.Concert.Replication.PutState.Misc", EAutomationTestFlags::EditorContext | EAutomationTestFlags::EngineFilter)
+		TUniquePtr<FObjectTestReplicator> ObjectReplicator;
+
+		TUniquePtr<FReplicationServer> Server;
+		FReplicationClient* Sender = nullptr;
+	
+		const FGuid StreamId = FGuid::NewGuid();
+	END_DEFINE_SPEC(FPutStateMiscSpec);
+
+	/** This tests misc workflows and cases with FConcertReplication_PutState_Request that do not really fit well with other specs. */
+	void FPutStateMiscSpec::Define()
+	{
+		BeforeEach([this]
+		{
+			ObjectReplicator = MakeUnique<FObjectTestReplicator>();
+			Server = MakeUnique<FReplicationServer>(*this);
+			Sender = &Server->ConnectClient();
+
+			Sender->JoinReplication();
+		});
+		AfterEach([this]
+		{
+			Server.Reset();
+			ObjectReplicator.Reset();
+		});
+
+		// This was a bug where the server would not remove the stream from the client.
+		Describe("When a put request removes client state", [this]
+		{
+			BeforeEach([this]
+			{
+				IConcertClientReplicationManager& ReplicationManager = Sender->GetClientReplicationManager();
+				ReplicationManager.ChangeStream({ .StreamsToAdd = { ObjectReplicator->CreateStream(StreamId) } });
+				ReplicationManager.PutClientState({ .NewStreams = { { Sender->GetEndpointId(), {} } } });
+			});
+			
+			It("The client state can create a new stream", [this]
+			{
+				bool bReceivedResponse = false;
+				Sender->GetClientReplicationManager()
+					.ChangeStream({ .StreamsToAdd = { ObjectReplicator->CreateStream(StreamId) } })
+					.Next([this, &bReceivedResponse](FConcertReplication_ChangeStream_Response&& Response)
+					{
+						bReceivedResponse = true;
+						TestTrue(TEXT("Success"), Response.IsSuccess());
+					});
+				TestTrue(TEXT("bReceivedResponse"), bReceivedResponse);
+			});
+
+			It("The stream has been fully deleted from the session", [this]
+			{
+				bool bReceivedResponse = false;
+				Sender->GetClientReplicationManager()
+					.QueryClientInfo({ .ClientEndpointIds = { Sender->GetEndpointId() } })
+					.Next([this, &bReceivedResponse](FConcertReplication_QueryReplicationInfo_Response&& Response)
+					{
+						bReceivedResponse = true;
+
+						const FConcertQueriedClientInfo* ClientInfo = Response.ClientInfo.Find(Sender->GetEndpointId());
+						TestTrue(TEXT("No streams"), ClientInfo && ClientInfo->Streams.IsEmpty());
+					});
+				TestTrue(TEXT("bReceivedResponse"), bReceivedResponse);
+			});
+		});
+	}
+}
diff --git a/Engine/Plugins/Developer/Concert/ConcertSync/ConcertSyncTest/Source/ConcertSyncTest/Private/Replication/PutState/PutStateRejection.spec.cpp b/Engine/Plugins/Developer/Concert/ConcertSync/ConcertSyncTest/Source/ConcertSyncTest/Private/Replication/PutState/PutStateRejection.spec.cpp
--- a/Engine/Plugins/Developer/Concert/ConcertSync/ConcertSyncTest/Source/ConcertSyncTest/Private/Replication/PutState/PutStateRejection.spec.cpp
+++ b/Engine/Plugins/Developer/Concert/ConcertSync/ConcertSyncTest/Source/ConcertSyncTest/Private/Replication/PutState/PutStateRejection.spec.cpp
@@ -265,5 +265,22 @@ namespace UE::ConcertSyncTests::Replication::ChangeClients
 				});
 			TestTrue(TEXT("Received response"), bReceivedResponse);
 		});
+
+		It("When request puts an empty stream, the request fails", [this]
+		{
+			IConcertClientReplicationManager& ReplicationManager = Client1->GetClientReplicationManager();
+
+			FConcertReplication_PutState_Request Request;
+			Request.NewStreams.Add(Client1->GetEndpointId(), { .Streams= { FConcertReplicationStream{} } });
+			bool bReceivedResponse = false;
+			ReplicationManager
+				.PutClientState(Request)
+				.Next([this, &bReceivedResponse](FConcertReplication_PutState_Response&& Response)
+				{
+					bReceivedResponse = true;
+					TestEqual(TEXT("Response code"), Response.ResponseCode, EConcertReplicationPutStateResponseCode::StreamError);
+				});
+			TestTrue(TEXT("bReceivedResponse"), bReceivedResponse);
+		});
 	}
 }
diff --git a/Engine/Plugins/Developer/Concert/ConcertSync/ConcertSyncTest/Source/ConcertSyncTest/Private/Replication/Stream/ChangeStream.spec.cpp b/Engine/Plugins/Developer/Concert/ConcertSync/ConcertSyncTest/Source/ConcertSyncTest/Private/Replication/Stream/ChangeStream.spec.cpp
--- a/Engine/Plugins/Developer/Concert/ConcertSync/ConcertSyncTest/Source/ConcertSyncTest/Private/Replication/Stream/ChangeStream.spec.cpp
+++ b/Engine/Plugins/Developer/Concert/ConcertSync/ConcertSyncTest/Source/ConcertSyncTest/Private/Replication/Stream/ChangeStream.spec.cpp
@@ -1,4 +1,4 @@
-﻿// Copyright Epic Games, Inc. All Rights Reserved.
+// Copyright Epic Games, Inc. All Rights Reserved.
 
 #include "Replication/Util/Mocks/ReplicationWorkspaceCallInterceptorMock.h"
 #include "Replication/Util/Spec/ReplicationClient.h"
@@ -79,5 +79,23 @@ namespace UE::ConcertSyncTests::Replication::RestoreContent
 				});
 			TestTrue(TEXT("bReceivedQueryResponse"), bReceivedStreamResponse);
 		});
+
+		It("When a request tries to create an empty stream, the request fails", [this]
+		{
+			AddExpectedError(TEXT("Rejecting ChangeStream request from"));
+			
+			bool bReceivedResponse = false;
+			IConcertClientReplicationManager& ReplicationManager = Client->GetClientReplicationManager();
+			ReplicationManager
+				.ChangeStream({ .StreamsToAdd = { FConcertReplicationStream{ .BaseDescription = { .Identifier = StreamId } } } })
+				.Next([this, &bReceivedResponse](FConcertReplication_ChangeStream_Response&& Response)
+				{
+					bReceivedResponse = true;
+					TestFalse(TEXT("Failure"), Response.IsSuccess());
+					TestEqual(TEXT("FailedStreamCreation.Num()"), Response.FailedStreamCreation.Num(), 1);
+					TestTrue(TEXT("FailedStreamCreation.Contains(StreamId)"), Response.FailedStreamCreation.Contains(StreamId));
+				});
+			TestTrue(TEXT("bReceivedResponse"), bReceivedResponse);
+		});
 	}
 }
