Introduce IConcertClientReplicationManager::OnPreRemoteEditApplied and IConcertClientReplicationManager::OnPostRemoteEditApplied events.
Also EConcertReplicationChangeClientReason to FConcertReplication_ChangeClientEvent.

diff --git a/Engine/Plugins/Developer/Concert/ConcertSync/ConcertSyncClient/Source/ConcertSyncClient/Private/Replication/Manager/ReplicationManager.cpp b/Engine/Plugins/Developer/Concert/ConcertSync/ConcertSyncClient/Source/ConcertSyncClient/Private/Replication/Manager/ReplicationManager.cpp
--- a/Engine/Plugins/Developer/Concert/ConcertSync/ConcertSyncClient/Source/ConcertSyncClient/Private/Replication/Manager/ReplicationManager.cpp
+++ b/Engine/Plugins/Developer/Concert/ConcertSync/ConcertSyncClient/Source/ConcertSyncClient/Private/Replication/Manager/ReplicationManager.cpp
@@ -156,46 +156,60 @@ namespace UE::ConcertSyncClient::Replication
 
 	IConcertClientReplicationManager::FOnPreStreamsChanged& FReplicationManager::OnPreStreamsChanged()
 	{
-		// Check() to avoid returning some dummy static variable
-		ensureMsgf(CurrentState, TEXT("StartAcceptingJoinRequests should have been called at this point."));
+		// checkf() to avoid returning some dummy static variable
+		checkf(CurrentState, TEXT("StartAcceptingJoinRequests should have been called at this point."));
 		return CurrentState->OnPreStreamsChanged();
 	}
 
 	IConcertClientReplicationManager::FOnPostStreamsChanged& FReplicationManager::OnPostStreamsChanged()
 	{
-		// Check() to avoid returning some dummy static variable
-		ensureMsgf(CurrentState, TEXT("StartAcceptingJoinRequests should have been called at this point."));
+		// checkf() to avoid returning some dummy static variable
+		checkf(CurrentState, TEXT("StartAcceptingJoinRequests should have been called at this point."));
 		return CurrentState->OnPostStreamsChanged();
 	}
 
 	IConcertClientReplicationManager::FOnPreAuthorityChanged& FReplicationManager::OnPreAuthorityChanged()
 	{
-		// Check() to avoid returning some dummy static variable
-		ensureMsgf(CurrentState, TEXT("StartAcceptingJoinRequests should have been called at this point."));
+		// checkf() to avoid returning some dummy static variable
+		checkf(CurrentState, TEXT("StartAcceptingJoinRequests should have been called at this point."));
 		return CurrentState->OnPreAuthorityChanged();
 	}
 
 	IConcertClientReplicationManager::FOnPostAuthorityChanged& FReplicationManager::OnPostAuthorityChanged()
 	{
-		// Check() to avoid returning some dummy static variable
-		ensureMsgf(CurrentState, TEXT("StartAcceptingJoinRequests should have been called at this point."));
+		// checkf() to avoid returning some dummy static variable
+		checkf(CurrentState, TEXT("StartAcceptingJoinRequests should have been called at this point."));
 		return CurrentState->OnPostAuthorityChanged();
 	}
 	
 	IConcertClientReplicationManager::FSyncControlChanged& FReplicationManager::OnPreSyncControlChanged()
 	{
-		// Check() to avoid returning some dummy static variable
-		ensureMsgf(CurrentState, TEXT("StartAcceptingJoinRequests should have been called at this point."));
+		// checkf() to avoid returning some dummy static variable
+		checkf(CurrentState, TEXT("StartAcceptingJoinRequests should have been called at this point."));
 		return CurrentState->OnPreSyncControlChanged();
 	}
 
 	IConcertClientReplicationManager::FSyncControlChanged& FReplicationManager::OnPostSyncControlChanged()
 	{
-		// Check() to avoid returning some dummy static variable
-		ensureMsgf(CurrentState, TEXT("StartAcceptingJoinRequests should have been called at this point."));
+		// checkf() to avoid returning some dummy static variable
+		checkf(CurrentState, TEXT("StartAcceptingJoinRequests should have been called at this point."));
 		return CurrentState->OnPostSyncControlChanged();
 	}
 
+	IConcertClientReplicationManager::FOnRemoteEditApplied& FReplicationManager::OnPreRemoteEditApplied()
+	{
+		// checkf() to avoid returning some dummy static variable
+		checkf(CurrentState, TEXT("StartAcceptingJoinRequests should have been called at this point."));
+		return CurrentState->OnPreRemoteEditApplied();
+	}
+
+	IConcertClientReplicationManager::FOnRemoteEditApplied& FReplicationManager::OnPostRemoteEditApplied()
+	{
+		// checkf() to avoid returning some dummy static variable
+		checkf(CurrentState, TEXT("StartAcceptingJoinRequests should have been called at this point."));
+		return CurrentState->OnPostRemoteEditApplied();
+	}
+
 	void FReplicationManager::OnChangeState(TSharedRef<FReplicationManagerState> NewState)
 	{
 		CurrentState = MoveTemp(NewState);
diff --git a/Engine/Plugins/Developer/Concert/ConcertSync/ConcertSyncClient/Source/ConcertSyncClient/Private/Replication/Manager/ReplicationManager.h b/Engine/Plugins/Developer/Concert/ConcertSync/ConcertSyncClient/Source/ConcertSyncClient/Private/Replication/Manager/ReplicationManager.h
--- a/Engine/Plugins/Developer/Concert/ConcertSync/ConcertSyncClient/Source/ConcertSyncClient/Private/Replication/Manager/ReplicationManager.h
+++ b/Engine/Plugins/Developer/Concert/ConcertSync/ConcertSyncClient/Source/ConcertSyncClient/Private/Replication/Manager/ReplicationManager.h
@@ -53,6 +53,8 @@ namespace UE::ConcertSyncClient::Replication
 		virtual FOnPostAuthorityChanged& OnPostAuthorityChanged() override;
 		virtual FSyncControlChanged& OnPreSyncControlChanged() override;
 		virtual FSyncControlChanged& OnPostSyncControlChanged() override;
+		virtual FOnRemoteEditApplied& OnPreRemoteEditApplied() override;
+		virtual FOnRemoteEditApplied& OnPostRemoteEditApplied() override;
 		//~ End IConcertClientReplicationManager Interface
 
 	private:
diff --git a/Engine/Plugins/Developer/Concert/ConcertSync/ConcertSyncClient/Source/ConcertSyncClient/Private/Replication/Manager/ReplicationManagerState.h b/Engine/Plugins/Developer/Concert/ConcertSync/ConcertSyncClient/Source/ConcertSyncClient/Private/Replication/Manager/ReplicationManagerState.h
--- a/Engine/Plugins/Developer/Concert/ConcertSync/ConcertSyncClient/Source/ConcertSyncClient/Private/Replication/Manager/ReplicationManagerState.h
+++ b/Engine/Plugins/Developer/Concert/ConcertSync/ConcertSyncClient/Source/ConcertSyncClient/Private/Replication/Manager/ReplicationManagerState.h
@@ -42,6 +42,8 @@ namespace UE::ConcertSyncClient::Replication
 		virtual FOnPostAuthorityChanged& OnPostAuthorityChanged() override { return OnPostAuthorityChangedDelegate; }
 		virtual FSyncControlChanged& OnPreSyncControlChanged() override { return OnPreSyncControlChangedDelegate; }
 		virtual FSyncControlChanged& OnPostSyncControlChanged() override { return OnPostSyncControlChangedDelegate; }
+		virtual FOnRemoteEditApplied& OnPreRemoteEditApplied() override { return OnPreRemoteEditAppliedDelegate; }
+		virtual FOnRemoteEditApplied& OnPostRemoteEditApplied() override { return OnPostRemoteEditAppliedDelegate; }
 		//~ End IConcertClientReplicationManager Interface
 
 	protected:
@@ -52,6 +54,8 @@ namespace UE::ConcertSyncClient::Replication
 		FOnPostAuthorityChanged OnPostAuthorityChangedDelegate;
 		FSyncControlChanged OnPreSyncControlChangedDelegate;
 		FSyncControlChanged OnPostSyncControlChangedDelegate;
+		FOnRemoteEditApplied OnPreRemoteEditAppliedDelegate;
+		FOnRemoteEditApplied OnPostRemoteEditAppliedDelegate;
 
 		/**
 		 * Subclasses can change the state with this function.
diff --git a/Engine/Plugins/Developer/Concert/ConcertSync/ConcertSyncClient/Source/ConcertSyncClient/Private/Replication/Manager/ReplicationManagerState_Connected_PutState.cpp b/Engine/Plugins/Developer/Concert/ConcertSync/ConcertSyncClient/Source/ConcertSyncClient/Private/Replication/Manager/ReplicationManagerState_Connected_PutState.cpp
--- a/Engine/Plugins/Developer/Concert/ConcertSync/ConcertSyncClient/Source/ConcertSyncClient/Private/Replication/Manager/ReplicationManagerState_Connected_PutState.cpp
+++ b/Engine/Plugins/Developer/Concert/ConcertSync/ConcertSyncClient/Source/ConcertSyncClient/Private/Replication/Manager/ReplicationManagerState_Connected_PutState.cpp
@@ -223,13 +223,18 @@ namespace UE::ConcertSyncClient::Replication
 			return;
 		}
 
-		const FConcertReplication_ChangeStream_Request& StreamChange = Event.StreamChange;
-		PredictAndApplyStreamChangeRemovedObjects(StreamChange);
-		FinalizePredictedStreamChange(StreamChange);
+		const FRemoteEditEvent EditEvent { Event.Reason, Event.ChangeData };
+		OnPreRemoteEditAppliedDelegate.Broadcast(EditEvent);
+		{
+			const FConcertReplication_ChangeStream_Request& StreamChange = Event.ChangeData.StreamChange;
+			PredictAndApplyStreamChangeRemovedObjects(StreamChange);
+			FinalizePredictedStreamChange(StreamChange);
 
-		const FConcertReplication_ChangeAuthority_Request& AuthorityChange = Event.AuthorityChange;
-		ApplyAuthorityChangeRemovedObjects(AuthorityChange);
-		FinalizePredictedAuthorityChange(AuthorityChange, {}, Event.SyncControlChange);
+			const FConcertReplication_ChangeAuthority_Request& AuthorityChange = Event.ChangeData.AuthorityChange;
+			ApplyAuthorityChangeRemovedObjects(AuthorityChange);
+			FinalizePredictedAuthorityChange(AuthorityChange, {}, Event.ChangeData.SyncControlChange);
+		}
+		OnPostRemoteEditAppliedDelegate.Broadcast(EditEvent);
 	}
 }
 
diff --git a/Engine/Plugins/Developer/Concert/ConcertSync/ConcertSyncClient/Source/ConcertSyncClient/Public/Replication/IConcertClientReplicationManager.h b/Engine/Plugins/Developer/Concert/ConcertSync/ConcertSyncClient/Source/ConcertSyncClient/Public/Replication/IConcertClientReplicationManager.h
--- a/Engine/Plugins/Developer/Concert/ConcertSync/ConcertSyncClient/Source/ConcertSyncClient/Public/Replication/IConcertClientReplicationManager.h
+++ b/Engine/Plugins/Developer/Concert/ConcertSync/ConcertSyncClient/Source/ConcertSyncClient/Public/Replication/IConcertClientReplicationManager.h
@@ -4,6 +4,7 @@
 
 #include "Misc/EBreakBehavior.h"
 #include "Replication/Messages/ChangeAuthority.h"
+#include "Replication/Messages/ChangeClientEvent.h"
 #include "Replication/Messages/PutState.h"
 #include "Replication/Messages/ChangeStream.h"
 #include "Replication/Messages/ClientQuery.h"
@@ -37,6 +38,12 @@ namespace UE::ConcertSyncClient::Replication
 			, DetailedErrorMessage(MoveTemp(DetailedErrorMessage))
 		{}
 	};
+
+	struct FRemoteEditEvent
+	{
+		const EConcertReplicationChangeClientReason Reason;
+		const FConcertReplication_ClientChangeData& ChangeData;
+	};
 }
 
 /**
@@ -165,6 +172,18 @@ public:
 	virtual FSyncControlChanged& OnPreSyncControlChanged() = 0;
 	/** Called just after sync control change is applied. */
 	virtual FSyncControlChanged& OnPostSyncControlChanged() = 0;
+
+	DECLARE_MULTICAST_DELEGATE_OneParam(FOnRemoteEditApplied, const UE::ConcertSyncClient::Replication::FRemoteEditEvent&);
+	/**
+	 * Called before updating any local state in response to the server notifying us that the local client's content was remotely edited.
+	 * This executes before OnPreStreamsChanged, OnPreAuthorityChanged, and OnPreSyncControlChanged, which may also get triggered depending on the change.
+	 */
+	virtual FOnRemoteEditApplied& OnPreRemoteEditApplied() = 0;
+	/**
+	 * Called after updating any local state in response to the server notifying us that the local client's content was remotely edited.
+	 * This executes after OnPreStreamsChanged, OnPreAuthorityChanged, and OnPreSyncControlChanged, which may also get triggered depending on the change.
+	 */
+	virtual FOnRemoteEditApplied& OnPostRemoteEditApplied() = 0;
 	
 	virtual ~IConcertClientReplicationManager() = default;
 };
diff --git a/Engine/Plugins/Developer/Concert/ConcertSync/ConcertSyncCore/Source/ConcertSyncCore/Public/Replication/Messages/ChangeClientEvent.h b/Engine/Plugins/Developer/Concert/ConcertSync/ConcertSyncCore/Source/ConcertSyncCore/Public/Replication/Messages/ChangeClientEvent.h
--- a/Engine/Plugins/Developer/Concert/ConcertSync/ConcertSyncCore/Source/ConcertSyncCore/Public/Replication/Messages/ChangeClientEvent.h
+++ b/Engine/Plugins/Developer/Concert/ConcertSync/ConcertSyncCore/Source/ConcertSyncCore/Public/Replication/Messages/ChangeClientEvent.h
@@ -6,18 +6,23 @@
 #include "ChangeStream.h"
 #include "ChangeClientEvent.generated.h"
 
-/**
- * Sent by server to notify a client that their stream content and / or authority has been changed by an external entity,
- * i.e. when the change was not initiated by the client itself.
- * 
- * For now, this is only in response to a remote client sending a FConcertReplication_PutState_Request but in the future it could also be
- * the user changing client content via the server UI.
- */
+/** A reason why a FConcertReplication_ChangeClientEvent is sent. */
+UENUM()
+enum class EConcertReplicationChangeClientReason : uint8
+{
+	/** The code is unknown */
+	Unknown,
+	
+	/** A FConcertReplication_PutState_Request edited this client. */
+	PutRequest
+};
+
+/** Describes what aspects about a client have changed. */
 USTRUCT()
-struct FConcertReplication_ChangeClientEvent
+struct FConcertReplication_ClientChangeData
 {
 	GENERATED_BODY()
-
+	
 	/** The change made to the client's streams. */
 	UPROPERTY()
 	FConcertReplication_ChangeStream_Request StreamChange;
@@ -29,4 +34,22 @@ struct FConcertReplication_ChangeClientEvent
 	/** The change made to the client's sync control in response to the above changes. */
 	UPROPERTY()
 	FConcertReplication_ChangeSyncControl SyncControlChange;
+};
+
+/**
+ * Sent by server to notify a client that their stream content and / or authority has been changed by an external entity,
+ * i.e. when the change was not initiated by the client itself.
+ */
+USTRUCT()
+struct FConcertReplication_ChangeClientEvent
+{
+	GENERATED_BODY()
+
+	/** The reason for this event */
+	UPROPERTY()
+	EConcertReplicationChangeClientReason Reason = EConcertReplicationChangeClientReason::Unknown;
+
+	/** Info about what has changed. */
+	UPROPERTY()
+	FConcertReplication_ClientChangeData ChangeData;
 };
\ No newline at end of file
diff --git a/Engine/Plugins/Developer/Concert/ConcertSync/ConcertSyncServer/Source/ConcertSyncServer/Private/Replication/ConcertServerReplicationManager.h b/Engine/Plugins/Developer/Concert/ConcertSync/ConcertSyncServer/Source/ConcertSyncServer/Private/Replication/ConcertServerReplicationManager.h
--- a/Engine/Plugins/Developer/Concert/ConcertSync/ConcertSyncServer/Source/ConcertSyncServer/Private/Replication/ConcertServerReplicationManager.h
+++ b/Engine/Plugins/Developer/Concert/ConcertSync/ConcertSyncServer/Source/ConcertSyncServer/Private/Replication/ConcertServerReplicationManager.h
@@ -16,6 +16,7 @@
 #include "SyncControlManager.h"
 
 #include "HAL/Platform.h"
+#include "Replication/Messages/ChangeClientEvent.h"
 #include "Replication/Messages/PutState.h"
 #include "Replication/Messages/RestoreContent.h"
 #include "Templates/SharedPointer.h"
@@ -154,15 +155,15 @@ namespace UE::ConcertSyncServer::Replication
 		void ApplyPutState_Streams(
 			const FGuid& RequestingEndpointId,
 			const TMap<FGuid, FConcertReplication_ChangeStream_Request> StreamRequests,
-			TMap<FGuid, FConcertReplication_ChangeClientEvent>& ClientChanges
-			);
+			TMap<FGuid, FConcertReplication_ClientChangeData>& ClientChanges
+		);
 		void ApplyPutState_Authority(
 			const FGuid& RequestingEndpointId,
 			const TSet<FConcertObjectInStreamID> RequestingClientSyncControlBefore,
 			const FConcertReplication_PutState_Request& Request,
 			FConcertReplication_PutState_Response& Response,
-			TMap<FGuid, FConcertReplication_ChangeClientEvent>& ClientChanges
-			);
+			TMap<FGuid, FConcertReplication_ClientChangeData>& ClientChanges
+		);
 		
 		// Leaving
 		void HandleLeaveReplicationSessionRequest(const FConcertSessionContext& ConcertSessionContext, const FConcertReplication_LeaveEvent& EventData);
diff --git a/Engine/Plugins/Developer/Concert/ConcertSync/ConcertSyncServer/Source/ConcertSyncServer/Private/Replication/ConcertServerReplicationManager_PutState.cpp b/Engine/Plugins/Developer/Concert/ConcertSync/ConcertSyncServer/Source/ConcertSyncServer/Private/Replication/ConcertServerReplicationManager_PutState.cpp
--- a/Engine/Plugins/Developer/Concert/ConcertSync/ConcertSyncServer/Source/ConcertSyncServer/Private/Replication/ConcertServerReplicationManager_PutState.cpp
+++ b/Engine/Plugins/Developer/Concert/ConcertSync/ConcertSyncServer/Source/ConcertSyncServer/Private/Replication/ConcertServerReplicationManager_PutState.cpp
@@ -220,7 +220,7 @@ namespace UE::ConcertSyncServer::Replication
 			const FGuid& RequestingEndpointId,
 			const FConcertReplication_ChangeMuteState_Request& Request,
 			FConcertReplication_PutState_Response& Response,
-			TMap<FGuid, FConcertReplication_ChangeClientEvent>& ClientChanges
+			TMap<FGuid, FConcertReplication_ClientChangeData>& ClientChanges
 			)
 		{
 			const bool bSuccess = MuteManager.ApplyRequestAndEnumerateSyncControl(
@@ -275,7 +275,7 @@ namespace UE::ConcertSyncServer::Replication
 		FConcertReplication_PutState_Response& Response
 		)
 	{
-		TMap<FGuid, FConcertReplication_ChangeClientEvent> ChangedSyncControl;
+		TMap<FGuid, FConcertReplication_ClientChangeData> ChangedSyncControl;
 		const TSet<FConcertObjectInStreamID> RequestingClientSyncControlBefore = [this, &RequestingEndpointId]()
 		{
 			const TSet<FConcertObjectInStreamID>* SyncControl = SyncControlManager.GetClientControlledObjects(RequestingEndpointId);
@@ -289,13 +289,14 @@ namespace UE::ConcertSyncServer::Replication
 			Private::ApplyPutState_Mute(MuteManager, RequestingEndpointId, Request.MuteChange, Response, ChangedSyncControl);
 		}
 
-		for (TPair<FGuid, FConcertReplication_ChangeClientEvent>& SyncControlPair : ChangedSyncControl)
+		for (TPair<FGuid, FConcertReplication_ClientChangeData>& SyncControlPair : ChangedSyncControl)
 		{
 			const FGuid& EndpointId = SyncControlPair.Key;
 			// The requesting client receives the sync control via Response (which was written by ApplyChangeClients_StreamPart, etc.)
 			if (EndpointId != RequestingEndpointId)
 			{
-				Session->SendCustomEvent(SyncControlPair.Value, EndpointId, EConcertMessageFlags::ReliableOrdered);
+				const FConcertReplication_ChangeClientEvent Event { EConcertReplicationChangeClientReason::PutRequest, MoveTemp(SyncControlPair.Value) };
+				Session->SendCustomEvent(Event, EndpointId, EConcertMessageFlags::ReliableOrdered);
 			}
 		}
 	}
@@ -303,8 +304,8 @@ namespace UE::ConcertSyncServer::Replication
 	void FConcertServerReplicationManager::ApplyPutState_Streams(
 		const FGuid& RequestingEndpointId,
 		const TMap<FGuid, FConcertReplication_ChangeStream_Request> StreamRequests,
-		TMap<FGuid, FConcertReplication_ChangeClientEvent>& ClientChanges
-		)
+		TMap<FGuid, FConcertReplication_ClientChangeData>& ClientChanges
+	)
 	{
 		for (const TPair<FGuid, FConcertReplication_ChangeStream_Request>& StreamPair : StreamRequests)
 		{
@@ -327,8 +328,8 @@ namespace UE::ConcertSyncServer::Replication
 		const TSet<FConcertObjectInStreamID> RequestingClientSyncControlBefore,
 		const FConcertReplication_PutState_Request& Request,
 		FConcertReplication_PutState_Response& Response,
-		TMap<FGuid, FConcertReplication_ChangeClientEvent>& ClientChanges
-		)
+		TMap<FGuid, FConcertReplication_ClientChangeData>& ClientChanges
+	)
 	{
 		for (const TPair<FGuid, FConcertObjectInStreamArray>& AuthorityPair : Request.NewAuthorityState)
 		{
diff --git a/Engine/Plugins/Developer/Concert/ConcertSync/ConcertSyncTest/Source/ConcertSyncTest/Private/Replication/PutState/ChangeClientEvent.spec.cpp b/Engine/Plugins/Developer/Concert/ConcertSync/ConcertSyncTest/Source/ConcertSyncTest/Private/Replication/PutState/ChangeClientEvent.spec.cpp
--- a/Engine/Plugins/Developer/Concert/ConcertSync/ConcertSyncTest/Source/ConcertSyncTest/Private/Replication/PutState/ChangeClientEvent.spec.cpp
+++ b/Engine/Plugins/Developer/Concert/ConcertSync/ConcertSyncTest/Source/ConcertSyncTest/Private/Replication/PutState/ChangeClientEvent.spec.cpp
@@ -45,9 +45,9 @@ namespace UE::ConcertSyncTests::Replication::ChangeClients
 				// This simulates the server sending to client - but it does not actually set any server state.
 				// This would not actually happen in a real world use case.
 				FConcertReplication_ChangeClientEvent Event;
-				Event.StreamChange.StreamsToAdd.Add(StreamData);
-				Event.AuthorityChange.TakeAuthority.Add({ ObjectReplicator->TestObject }, { { StreamId } });
-				Event.SyncControlChange.NewControlStates.Add({ StreamId, ObjectReplicator->TestObject }, true);
+				Event.ChangeData.StreamChange.StreamsToAdd.Add(StreamData);
+				Event.ChangeData.AuthorityChange.TakeAuthority.Add({ ObjectReplicator->TestObject }, { { StreamId } });
+				Event.ChangeData.SyncControlChange.NewControlStates.Add({ StreamId, ObjectReplicator->TestObject }, true);
 				Server->GetServerSessionMock()->SendCustomEvent(Event, Client->GetEndpointId(), EConcertMessageFlags::ReliableOrdered);
 			});
 
diff --git a/Engine/Plugins/Developer/Concert/ConcertSync/ConcertSyncTest/Source/ConcertSyncTest/Private/Replication/PutState/PutStateAndMute.spec.cpp b/Engine/Plugins/Developer/Concert/ConcertSync/ConcertSyncTest/Source/ConcertSyncTest/Private/Replication/PutState/PutStateAndMute.spec.cpp
--- a/Engine/Plugins/Developer/Concert/ConcertSync/ConcertSyncTest/Source/ConcertSyncTest/Private/Replication/PutState/PutStateAndMute.spec.cpp
+++ b/Engine/Plugins/Developer/Concert/ConcertSync/ConcertSyncTest/Source/ConcertSyncTest/Private/Replication/PutState/PutStateAndMute.spec.cpp
@@ -238,7 +238,7 @@ namespace UE::ConcertSyncTests::Replication::ChangeClients
 				ensureMsgf(ReplicationManager.GetSyncControlledObjects().Contains({ StreamId, ObjectReplicator->TestObject }), TEXT("Test not set up correctly."));
 			});
 			
-			It("When client is loses sync control due to a mute change.", [this]
+			It("When client loses sync control due to a mute change.", [this]
 			{
 				int32 EventCount = 0;
 				Sender->GetClientSessionMock()->RegisterCustomEventHandler<FConcertReplication_ChangeClientEvent>(
@@ -246,8 +246,9 @@ namespace UE::ConcertSyncTests::Replication::ChangeClients
 					{
 						++EventCount;
 						
-						TestEqual(TEXT("NewControlStates.Num()"), Event.SyncControlChange.NewControlStates.Num(), 1);
-						const bool* NewSyncControl = Event.SyncControlChange.NewControlStates.Find({ StreamId, ObjectReplicator->TestObject });
+						TestEqual(TEXT("Reason"), Event.Reason, EConcertReplicationChangeClientReason::PutRequest);
+						TestEqual(TEXT("NewControlStates.Num()"), Event.ChangeData.SyncControlChange.NewControlStates.Num(), 1);
+						const bool* NewSyncControl = Event.ChangeData.SyncControlChange.NewControlStates.Find({ StreamId, ObjectReplicator->TestObject });
 						TestTrue(TEXT("Has No Sync Control"), NewSyncControl && !*NewSyncControl);
 					});
 				
@@ -258,7 +259,7 @@ namespace UE::ConcertSyncTests::Replication::ChangeClients
 
 				TestEqual(TEXT("EventCount"), EventCount, 1);
 			});
-			It("When client is gains sync control due to a mute change.", [this]
+			It("When client gains sync control due to a mute change.", [this]
 			{
 				int32 EventCount = 0;
 				Sender->GetClientSessionMock()->RegisterCustomEventHandler<FConcertReplication_ChangeClientEvent>(
@@ -266,8 +267,9 @@ namespace UE::ConcertSyncTests::Replication::ChangeClients
 					{
 						++EventCount;
 						
-						TestEqual(TEXT("NewControlStates.Num()"), Event.SyncControlChange.NewControlStates.Num(), 1);
-						const bool* NewSyncControl = Event.SyncControlChange.NewControlStates.Find({ StreamId, ObjectReplicator->TestObject });
+						TestEqual(TEXT("Reason"), Event.Reason, EConcertReplicationChangeClientReason::PutRequest);
+						TestEqual(TEXT("NewControlStates.Num()"), Event.ChangeData.SyncControlChange.NewControlStates.Num(), 1);
+						const bool* NewSyncControl = Event.ChangeData.SyncControlChange.NewControlStates.Find({ StreamId, ObjectReplicator->TestObject });
 						TestTrue(TEXT("Has No Sync Control"), NewSyncControl && *NewSyncControl);
 					});
 				
@@ -283,15 +285,10 @@ namespace UE::ConcertSyncTests::Replication::ChangeClients
 		
 		It("When client is not affected by a mute change, it does not receive FConcertReplication_ChangeClientEvent.", [this]
 		{
-			int32 EventCount = 0;
 			Receiver->GetClientSessionMock()->RegisterCustomEventHandler<FConcertReplication_ChangeClientEvent>(
-				[this, &EventCount](const FConcertSessionContext&, const FConcertReplication_ChangeClientEvent& Event)
+				[this](const FConcertSessionContext&, const FConcertReplication_ChangeClientEvent& Event)
 				{
-					++EventCount;
-					
-					TestEqual(TEXT("NewControlStates.Num()"), Event.SyncControlChange.NewControlStates.Num(), 1);
-					const bool* NewSyncControl = Event.SyncControlChange.NewControlStates.Find({ StreamId, ObjectReplicator->TestObject });
-					TestTrue(TEXT("Has No Sync Control"), NewSyncControl && *NewSyncControl);
+					AddError(TEXT("Event was not expected"));
 				});
 			
 			FObjectTestReplicator BarReplicator;
@@ -304,8 +301,6 @@ namespace UE::ConcertSyncTests::Replication::ChangeClients
 			Request.NewStreams.Add(Sender->GetEndpointId(), {{ ObjectReplicator->CreateStream(StreamId) }});
 			Request.MuteChange.ObjectsToMute.Add(ObjectReplicator->TestObject);
 			SenderManager.PutClientState(Request);
-
-			TestEqual(TEXT("EventCount"), EventCount, 0);
 		});
 	}
 }
diff --git a/Engine/Plugins/Developer/Concert/ConcertSync/ConcertSyncTest/Source/ConcertSyncTest/Private/Replication/PutState/PutStateMisc.spec.cpp b/Engine/Plugins/Developer/Concert/ConcertSync/ConcertSyncTest/Source/ConcertSyncTest/Private/Replication/PutState/PutStateMisc.spec.cpp
--- a/Engine/Plugins/Developer/Concert/ConcertSync/ConcertSyncTest/Source/ConcertSyncTest/Private/Replication/PutState/PutStateMisc.spec.cpp
+++ b/Engine/Plugins/Developer/Concert/ConcertSync/ConcertSyncTest/Source/ConcertSyncTest/Private/Replication/PutState/PutStateMisc.spec.cpp
@@ -1,11 +1,13 @@
 ﻿// Copyright Epic Games, Inc. All Rights Reserved.
 
+#include "Algo/Count.h"
 #include "Replication/Util/Spec/ReplicationClient.h"
 #include "Replication/Util/Spec/ReplicationServer.h"
 #include "Replication/Util/Spec/ObjectTestReplicator.h"
 
 #include "Misc/AutomationTest.h"
 #include "Replication/Messages/PutState.h"
+#include "Replication/Util/ClientEventRecorder.h"
 
 namespace UE::ConcertSyncTests::Replication::ChangeClients
 {
@@ -13,7 +15,8 @@ namespace UE::ConcertSyncTests::Replication::ChangeClients
 		TUniquePtr<FObjectTestReplicator> ObjectReplicator;
 
 		TUniquePtr<FReplicationServer> Server;
-		FReplicationClient* Sender = nullptr;
+		FReplicationClient* Client1 = nullptr;
+		FReplicationClient* Client2 = nullptr;
 	
 		const FGuid StreamId = FGuid::NewGuid();
 	END_DEFINE_SPEC(FPutStateMiscSpec);
@@ -25,9 +28,11 @@ namespace UE::ConcertSyncTests::Replication::ChangeClients
 		{
 			ObjectReplicator = MakeUnique<FObjectTestReplicator>();
 			Server = MakeUnique<FReplicationServer>(*this);
-			Sender = &Server->ConnectClient();
+			Client1 = &Server->ConnectClient();
+			Client2 = &Server->ConnectClient();
 
-			Sender->JoinReplication();
+			Client1->JoinReplication();
+			Client2->JoinReplication();
 		});
 		AfterEach([this]
 		{
@@ -40,15 +45,15 @@ namespace UE::ConcertSyncTests::Replication::ChangeClients
 		{
 			BeforeEach([this]
 			{
-				IConcertClientReplicationManager& ReplicationManager = Sender->GetClientReplicationManager();
+				IConcertClientReplicationManager& ReplicationManager = Client1->GetClientReplicationManager();
 				ReplicationManager.ChangeStream({ .StreamsToAdd = { ObjectReplicator->CreateStream(StreamId) } });
-				ReplicationManager.PutClientState({ .NewStreams = { { Sender->GetEndpointId(), {} } } });
+				ReplicationManager.PutClientState({ .NewStreams = { { Client1->GetEndpointId(), {} } } });
 			});
 			
 			It("The client state can create a new stream", [this]
 			{
 				bool bReceivedResponse = false;
-				Sender->GetClientReplicationManager()
+				Client1->GetClientReplicationManager()
 					.ChangeStream({ .StreamsToAdd = { ObjectReplicator->CreateStream(StreamId) } })
 					.Next([this, &bReceivedResponse](FConcertReplication_ChangeStream_Response&& Response)
 					{
@@ -61,17 +66,60 @@ namespace UE::ConcertSyncTests::Replication::ChangeClients
 			It("The stream has been fully deleted from the session", [this]
 			{
 				bool bReceivedResponse = false;
-				Sender->GetClientReplicationManager()
-					.QueryClientInfo({ .ClientEndpointIds = { Sender->GetEndpointId() } })
+				Client1->GetClientReplicationManager()
+					.QueryClientInfo({ .ClientEndpointIds = { Client1->GetEndpointId() } })
 					.Next([this, &bReceivedResponse](FConcertReplication_QueryReplicationInfo_Response&& Response)
 					{
 						bReceivedResponse = true;
 
-						const FConcertQueriedClientInfo* ClientInfo = Response.ClientInfo.Find(Sender->GetEndpointId());
+						const FConcertQueriedClientInfo* ClientInfo = Response.ClientInfo.Find(Client1->GetEndpointId());
 						TestTrue(TEXT("No streams"), ClientInfo && ClientInfo->Streams.IsEmpty());
 					});
 				TestTrue(TEXT("bReceivedResponse"), bReceivedResponse);
 			});
 		});
+
+		It("OnPreRemoteEditApplied and OnPostRemoteEditApplied are triggered in right order", [this]
+		{
+			const FClientEventRecorder EventRecorder(Client1->GetClientReplicationManager());
+			FConcertReplication_PutState_Request Request;
+			Request.NewStreams.Add(Client1->GetEndpointId(), {{ ObjectReplicator->CreateStream(StreamId) }});
+			Request.NewAuthorityState.Add(Client1->GetEndpointId()).Objects.Add({ StreamId, ObjectReplicator->TestObject });
+			Client2->GetClientReplicationManager().PutClientState(Request);
+
+			const TArray<EEventType>& ActualOrder = EventRecorder.GetEventOrder();
+			if (ActualOrder.Num() < 8) // There is Pre & Post for the 4 events: Stream, Authority, SyncControl, RemoteEdit
+			{
+				AddError(TEXT("No all events included"));
+				return;
+			}
+			
+			TestEqual(TEXT("PreRemoteEditApplied comes first"), ActualOrder[0], EEventType::PreRemoteEditApplied);
+			TestEqual(TEXT("PreRemoteEditApplied comes last"), ActualOrder[ActualOrder.Num() - 1], EEventType::PostRemoteEditApplied);
+
+			const int32 NumPreRemoteEditAppliedBroadcasts = Algo::Count(ActualOrder, EEventType::PreRemoteEditApplied);
+			const int32 NumPostRemoteEditAppliedBroadcasts = Algo::Count(ActualOrder, EEventType::PostRemoteEditApplied);
+			TestEqual(TEXT("PreRemoteEditApplied appears once"), NumPreRemoteEditAppliedBroadcasts, 1);
+			TestEqual(TEXT("PostRemoteEditApplied appears once"), NumPostRemoteEditAppliedBroadcasts, 1);
+		});
+		It("OnPreRemoteEditApplied and OnPostRemoteEditApplied have right reason", [this]
+		{
+			int32 EventCount = 0;
+			IConcertClientReplicationManager& ReplicationManager = Client1->GetClientReplicationManager();
+			const auto HandleEvent = [this, &EventCount](const ConcertSyncClient::Replication::FRemoteEditEvent& Event)
+			{
+				++EventCount;
+				TestEqual(TEXT("Reason"), Event.Reason, EConcertReplicationChangeClientReason::PutRequest);
+			};
+			ReplicationManager.OnPreRemoteEditApplied().AddLambda(HandleEvent);
+			ReplicationManager.OnPostRemoteEditApplied().AddLambda(HandleEvent);
+			
+			FConcertReplication_PutState_Request Request;
+			Request.NewStreams.Add(Client1->GetEndpointId(), {{ ObjectReplicator->CreateStream(StreamId) }});
+			Request.NewAuthorityState.Add(Client1->GetEndpointId()).Objects.Add({ StreamId, ObjectReplicator->TestObject });
+			Client2->GetClientReplicationManager().PutClientState(Request);
+
+			TestEqual(TEXT("EventCount"), EventCount, 2);
+		});
 	}
 }
diff --git a/Engine/Plugins/Developer/Concert/ConcertSync/ConcertSyncTest/Source/ConcertSyncTest/Private/Replication/Util/ClientEventRecorder.h b/Engine/Plugins/Developer/Concert/ConcertSync/ConcertSyncTest/Source/ConcertSyncTest/Private/Replication/Util/ClientEventRecorder.h
new file mode 100644
--- /dev/null
+++ b/Engine/Plugins/Developer/Concert/ConcertSync/ConcertSyncTest/Source/ConcertSyncTest/Private/Replication/Util/ClientEventRecorder.h
@@ -0,0 +1,62 @@
+﻿// Copyright Epic Games, Inc. All Rights Reserved.
+
+#pragma once
+
+#include "HAL/Platform.h"
+#include "Replication/IConcertClientReplicationManager.h"
+#include "Templates/UnrealTemplate.h"
+
+namespace UE::ConcertSyncTests::Replication
+{
+	enum class EEventType
+	{
+		PreStreamChange,
+		PostStreamChange,
+		PreAuthorityChange,
+		PostAuthorityChange,
+		PreSyncControlChange,
+		PostSyncControlChange,
+		PreRemoteEditApplied,
+		PostRemoteEditApplied
+	};
+
+	/** Records the order in which events are broadcast on an IConcertClientReplicationManager. */
+	class FClientEventRecorder : public FNoncopyable
+	{
+		IConcertClientReplicationManager& ReplicationManager;
+		TArray<EEventType> EventOrder;
+	public:
+
+		FClientEventRecorder(IConcertClientReplicationManager& ReplicationManager UE_LIFETIMEBOUND)
+			: ReplicationManager(ReplicationManager)
+		{
+			ReplicationManager.OnPreStreamsChanged().AddRaw(this, &FClientEventRecorder::HandleEvent, EEventType::PreStreamChange);
+			ReplicationManager.OnPostStreamsChanged().AddRaw(this, &FClientEventRecorder::HandleEvent, EEventType::PostStreamChange);
+			ReplicationManager.OnPreAuthorityChanged().AddRaw(this, &FClientEventRecorder::HandleEvent, EEventType::PreAuthorityChange);
+			ReplicationManager.OnPostAuthorityChanged().AddRaw(this, &FClientEventRecorder::HandleEvent, EEventType::PostAuthorityChange);
+			ReplicationManager.OnPreSyncControlChanged().AddRaw(this, &FClientEventRecorder::HandleEvent, EEventType::PreSyncControlChange);
+			ReplicationManager.OnPostSyncControlChanged().AddRaw(this, &FClientEventRecorder::HandleEvent, EEventType::PostSyncControlChange);
+			ReplicationManager.OnPreRemoteEditApplied().AddRaw(this, &FClientEventRecorder::HandleEvent, EEventType::PreRemoteEditApplied);
+			ReplicationManager.OnPostRemoteEditApplied().AddRaw(this, &FClientEventRecorder::HandleEvent, EEventType::PostRemoteEditApplied);
+		}
+		~FClientEventRecorder()
+		{
+			ReplicationManager.OnPreStreamsChanged().RemoveAll(this);
+			ReplicationManager.OnPostStreamsChanged().RemoveAll(this);
+			ReplicationManager.OnPreAuthorityChanged().RemoveAll(this);
+			ReplicationManager.OnPostAuthorityChanged().RemoveAll(this);
+			ReplicationManager.OnPreSyncControlChanged().RemoveAll(this);
+			ReplicationManager.OnPostSyncControlChanged().RemoveAll(this);
+			ReplicationManager.OnPreRemoteEditApplied().RemoveAll(this);
+			ReplicationManager.OnPostRemoteEditApplied().RemoveAll(this);
+		}
+
+		const TArray<EEventType>& GetEventOrder() const { return EventOrder; }
+		void Clear() { EventOrder.Empty(); }
+
+	private:
+
+		void HandleEvent(const EEventType EventType) { EventOrder.Add(EventType); }
+		void HandleEvent(const ConcertSyncClient::Replication::FRemoteEditEvent&, const EEventType EventType) { EventOrder.Add(EventType);  }
+	};
+}
