Add EConcertReplicationMuteRequestFlags::ClearMuteState to FConcertReplication_ChangeMuteState_Request, which allows resetting the entire mute state in a single request.

diff --git a/Engine/Plugins/Developer/Concert/ConcertSync/ConcertSyncCore/Source/ConcertSyncCore/Public/Replication/Messages/Muting.h b/Engine/Plugins/Developer/Concert/ConcertSync/ConcertSyncCore/Source/ConcertSyncCore/Public/Replication/Messages/Muting.h
--- a/Engine/Plugins/Developer/Concert/ConcertSync/ConcertSyncCore/Source/ConcertSyncCore/Public/Replication/Messages/Muting.h
+++ b/Engine/Plugins/Developer/Concert/ConcertSync/ConcertSyncCore/Source/ConcertSyncCore/Public/Replication/Messages/Muting.h
@@ -37,6 +37,16 @@ struct FConcertReplication_ObjectMuteSetting
 	friend bool operator==(const FConcertReplication_ObjectMuteSetting& Left, const FConcertReplication_ObjectMuteSetting& Right) = default;
 };
 
+UENUM(Flags)
+enum class EConcertReplicationMuteRequestFlags : uint8
+{
+	None,
+
+	/** Before the request is applied, all mute state is reset. This means that ObjectsToUnmute is effectively ignored and the mute state is replaced with what is saved in ObjectsToMute. */
+	ClearMuteState = 1 << 0,
+};
+ENUM_CLASS_FLAGS(EConcertReplicationMuteRequestFlags);
+
 /**
  * A request to globally pause / resume replication of objects.
  * Muted objects will not be replicated to any clients.
@@ -66,6 +76,10 @@ struct FConcertReplication_ChangeMuteState_Request
 {
 	GENERATED_BODY()
 
+	/** Flags that modify the request's behavior. */
+	UPROPERTY()
+	EConcertReplicationMuteRequestFlags Flags = EConcertReplicationMuteRequestFlags::None;
+	
 	/** The objects to explicitly mute. */
 	UPROPERTY()
 	TMap<FSoftObjectPath, FConcertReplication_ObjectMuteSetting> ObjectsToMute;
@@ -79,6 +93,8 @@ struct FConcertReplication_ChangeMuteState_Request
 	 *
 	 * If a specified object is explicitly muted, then it and all the subobjects that were implicitly muted because of it are unmuted;
 	 * in this case, it does not matter whether EConcertReplicationMuteFlags::ObjectAndSubobjects flag is set.
+	 *
+	 * If Flags specifies ClearMuteState and ObjectsToUnmute is non-empty, the request is rejected.
 	 */
 	UPROPERTY()
 	TMap<FSoftObjectPath, FConcertReplication_ObjectMuteSetting> ObjectsToUnmute;
diff --git a/Engine/Plugins/Developer/Concert/ConcertSync/ConcertSyncServer/Source/ConcertSyncServer/Private/Replication/Muting/MuteManager.cpp b/Engine/Plugins/Developer/Concert/ConcertSync/ConcertSyncServer/Source/ConcertSyncServer/Private/Replication/Muting/MuteManager.cpp
--- a/Engine/Plugins/Developer/Concert/ConcertSync/ConcertSyncServer/Source/ConcertSyncServer/Private/Replication/Muting/MuteManager.cpp
+++ b/Engine/Plugins/Developer/Concert/ConcertSync/ConcertSyncServer/Source/ConcertSyncServer/Private/Replication/Muting/MuteManager.cpp
@@ -295,9 +295,13 @@ namespace UE::ConcertSyncServer::Replication
 		TFunctionRef<void(const FSoftObjectPath& ObjectPath)> OnRejection
 		) const
 	{
+		if (EnumHasAnyFlags(Request.Flags, EConcertReplicationMuteRequestFlags::ClearMuteState) && !Request.ObjectsToUnmute.IsEmpty())
+		{
+			return false;
+		}
+		
 		bool bIsValidRequest = true;
 		const auto Reject = [&OnRejection, &bIsValidRequest](const FSoftObjectPath& ObjectPath){ OnRejection(ObjectPath); bIsValidRequest = false; };
-		
 		const auto SharedValidateSetting = [this, &Reject, &ObjectCache](const FSoftObjectPath& Object, const FConcertReplication_ObjectMuteSetting& Setting)
 		{
 			const bool bIsObjectReferenced = ObjectCache.IsObjectReferencedDirectly(Object);
@@ -339,7 +343,14 @@ namespace UE::ConcertSyncServer::Replication
 
 	void FMuteManager::ApplyRequestInternal(const FConcertReplication_ChangeMuteState_Request& Request)
 	{
-		UE_LOG(LogConcert, Log, TEXT("Explicitly muting %d and unmuting %d objects."), Request.ObjectsToMute.Num(), Request.ObjectsToUnmute.Num());
+		const bool bResetState = EnumHasAnyFlags(Request.Flags, EConcertReplicationMuteRequestFlags::ClearMuteState);
+		UE_CLOG(bResetState, LogConcert, Log, TEXT("Explicitly muting %d and unmuting %d objects."), Request.ObjectsToMute.Num(), Request.ObjectsToUnmute.Num());
+		UE_CLOG(!bResetState, LogConcert, Log, TEXT("Explicitly muting %d and unmuting %d objects. Resetting mute state before."), Request.ObjectsToMute.Num(), Request.ObjectsToUnmute.Num());
+		
+		if (bResetState)
+		{
+			MuteState.Empty(Request.ObjectsToMute.Num());
+		}
 		
 		for (const TPair<FSoftObjectPath, FConcertReplication_ObjectMuteSetting>& ToMute : Request.ObjectsToMute)
 		{
diff --git a/Engine/Plugins/Developer/Concert/ConcertSync/ConcertSyncTest/Source/ConcertSyncTest/Private/Replication/Muting/ResetMuteState.spec.cpp b/Engine/Plugins/Developer/Concert/ConcertSync/ConcertSyncTest/Source/ConcertSyncTest/Private/Replication/Muting/ResetMuteState.spec.cpp
new file mode 100644
--- /dev/null
+++ b/Engine/Plugins/Developer/Concert/ConcertSync/ConcertSyncTest/Source/ConcertSyncTest/Private/Replication/Muting/ResetMuteState.spec.cpp
@@ -0,0 +1,196 @@
+﻿// Copyright Epic Games, Inc. All Rights Reserved.
+
+#include "Misc/AutomationTest.h"
+#include "Replication/Util/ConcertClientReplicationBridgeMock.h"
+#include "Replication/Util/Spec/ObjectTestReplicator.h"
+#include "Replication/Util/Spec/ReplicationClient.h"
+#include "Replication/Util/Spec/ReplicationServer.h"
+
+namespace UE::ConcertSyncTests::Replication
+{
+	/** Tests that EConcertReplicationMuteRequestFlags::ResetMuteState works as intended in conjunction with FConcertReplication_ChangeMuteState_Request. */
+	BEGIN_DEFINE_SPEC(FResetMuteStateSpec, "Editor.Concert.Replication.Muting.Reset", EAutomationTestFlags::EditorContext | EAutomationTestFlags::EngineFilter)
+		TUniquePtr<FObjectTestReplicator> ObjectReplicator;
+		TUniquePtr<FObjectTestReplicator> SecondObjectReplicator;
+
+		TUniquePtr<FReplicationServer> Server;
+		FReplicationClient* Sender = nullptr;
+		FReplicationClient* Receiver = nullptr;
+
+		const FGuid StreamId = FGuid::NewGuid();
+	END_DEFINE_SPEC(FResetMuteStateSpec);
+
+	void FResetMuteStateSpec::Define()
+	{
+		BeforeEach([this]
+		{
+			ObjectReplicator		= MakeUnique<FObjectTestReplicator>(TEXT("ObjectReplicator"));
+			SecondObjectReplicator	= MakeUnique<FObjectTestReplicator>(TEXT("SecondObjectReplicator"));
+			Server					= MakeUnique<FReplicationServer>(*this);
+			Sender					= &Server->ConnectClient();
+			Receiver				= &Server->ConnectClient();
+
+			Sender->JoinReplication(ObjectReplicator->CreateSenderArgs(StreamId));
+			Receiver->JoinReplicationAsListener({ ObjectReplicator->TestObject });
+			
+			Sender->GetClientReplicationManager().TakeAuthorityOver({  ObjectReplicator->TestObject });
+			Sender->GetClientReplicationManager().ChangeMuteState({ .ObjectsToMute = { { ObjectReplicator->TestObject, {} } } });
+			// Needs to be known by server, so it can be muted
+			Receiver->GetClientReplicationManager().ChangeStream({ .StreamsToAdd = { { SecondObjectReplicator->CreateStream(StreamId) } } });
+		});
+		AfterEach([this]
+		{
+			Server.Reset();
+			ObjectReplicator.Reset();
+		});
+
+		It("When EConcertReplicationMuteRequestFlags::ClearMuteState is specified and ObjectsToUnmute is non-empty, the request is rejected.", [this]
+		{
+			int32 EventCount = 0;
+			Sender->GetClientReplicationManager()
+				.ChangeMuteState({ .Flags = EConcertReplicationMuteRequestFlags::ClearMuteState, .ObjectsToUnmute = { { ObjectReplicator->TestObject, {} } } })
+				.Next([this, &EventCount](FConcertReplication_ChangeMuteState_Response&& Response)
+				{
+					++EventCount;
+					TestEqual(TEXT("Success"), Response.ErrorCode, EConcertReplicationMuteErrorCode::Rejected);
+					return Response;
+				});
+			TestEqual(TEXT("EventCount"), EventCount, 1);
+		});
+
+		Describe("When EConcertReplicationMuteRequestFlags::ResetMuteState is used on its own", [this]
+		{
+			const auto RunTest = [this]()
+			{
+				int32 EventCount = 0;
+				TFuture<FConcertReplication_ChangeMuteState_Response> Future = Sender->GetClientReplicationManager()
+					.ChangeMuteState({ .Flags = EConcertReplicationMuteRequestFlags::ClearMuteState })
+					.Next([this, &EventCount](FConcertReplication_ChangeMuteState_Response&& Response)
+					{
+						++EventCount;
+						TestTrue(TEXT("Success"), Response.IsSuccess());
+						return Response;
+					});
+				TestEqual(TEXT("EventCount"), EventCount, 1);
+				return Future;
+			};
+
+			It("The response contains sync control", [this, RunTest]
+			{
+				RunTest()
+				.Next([this](FConcertReplication_ChangeMuteState_Response&& Response)
+				{
+					TestEqual(TEXT("NewControlStates.Num()"), Response.SyncControl.NewControlStates.Num(), 1);
+					const bool* NewState = Response.SyncControl.NewControlStates.Find({ StreamId, ObjectReplicator->TestObject });
+					TestTrue(TEXT("Has Sync Control"), NewState && *NewState);
+				});
+			});
+			It("Server mute state is empty", [this, RunTest]
+			{
+				RunTest();
+				
+				int32 EventCount = 0;
+				Sender->GetClientReplicationManager()
+					.QueryMuteState()
+					.Next([this, &EventCount](FConcertReplication_QueryMuteState_Response&& Response)
+					{
+						++EventCount;
+						TestTrue(TEXT("IsEmpty"), Response.IsEmpty());
+					});
+				TestEqual(TEXT("EventCount"), EventCount, 1);
+			});
+		});
+
+		Describe("When EConcertReplicationMuteRequestFlags::ResetMuteState is used and ObjectsToMute mutes the same object again", [this]
+		{
+			const auto RunTest = [this]()
+			{
+				int32 EventCount = 0;
+				TFuture<FConcertReplication_ChangeMuteState_Response> Future = Sender->GetClientReplicationManager()
+					.ChangeMuteState({ .Flags = EConcertReplicationMuteRequestFlags::ClearMuteState, .ObjectsToMute = { { ObjectReplicator->TestObject, {} } } })
+					.Next([this, &EventCount](FConcertReplication_ChangeMuteState_Response&& Response)
+					{
+						++EventCount;
+						TestTrue(TEXT("Success"), Response.IsSuccess());
+						return Response;
+					});
+				TestEqual(TEXT("EventCount"), EventCount, 1);
+				return Future;
+			};
+
+			It("The response contains no sync control", [this, RunTest]
+			{
+				RunTest()
+				.Next([this](FConcertReplication_ChangeMuteState_Response&& Response)
+				{
+					TestEqual(TEXT("NewControlStates.Num()"), Response.SyncControl.NewControlStates.Num(), 0);
+				});
+			});
+			It("Server mute state is set", [this, RunTest]
+			{
+				RunTest();
+				
+				int32 EventCount = 0;
+				Sender->GetClientReplicationManager()
+					.QueryMuteState()
+					.Next([this, &EventCount](FConcertReplication_QueryMuteState_Response&& Response)
+					{
+						++EventCount;
+						TestEqual(TEXT("ExplicitlyMutedObjects.Num()"), Response.ExplicitlyMutedObjects.Num(), 1);
+						TestEqual(TEXT("ExplicitlyUnmutedObjects.Num()"), Response.ExplicitlyUnmutedObjects.Num(), 0);
+						TestEqual(TEXT("ImplicitlyMutedObjects.Num()"), Response.ImplicitlyMutedObjects.Num(), 0);
+						TestEqual(TEXT("ImplicitlyUnmutedObjects.Num()"), Response.ImplicitlyUnmutedObjects.Num(), 0);
+						TestTrue(TEXT("Object is muted"), Response.ExplicitlyMutedObjects.Contains(ObjectReplicator->TestObject));
+					});
+				TestEqual(TEXT("EventCount"), EventCount, 1);
+			});
+		});
+		
+		Describe("When EConcertReplicationMuteRequestFlags::ResetMuteState is used and ObjectsToMute mutes another object", [this]
+		{
+			const auto RunTest = [this]()
+			{
+				int32 EventCount = 0;
+				TFuture<FConcertReplication_ChangeMuteState_Response> Future = Sender->GetClientReplicationManager()
+					.ChangeMuteState({ .Flags = EConcertReplicationMuteRequestFlags::ClearMuteState, .ObjectsToMute = { { SecondObjectReplicator->TestObject, {} } } })
+					.Next([this, &EventCount](FConcertReplication_ChangeMuteState_Response&& Response)
+					{
+						++EventCount;
+						TestTrue(TEXT("Success"), Response.IsSuccess());
+						return Response;
+					});
+				TestEqual(TEXT("EventCount"), EventCount, 1);
+				return Future;
+			};
+
+			It("The response contains sync control", [this, RunTest]
+			{
+				RunTest()
+				.Next([this](FConcertReplication_ChangeMuteState_Response&& Response)
+				{
+					TestEqual(TEXT("NewControlStates.Num()"), Response.SyncControl.NewControlStates.Num(), 1);
+					const bool* NewState = Response.SyncControl.NewControlStates.Find({ StreamId, ObjectReplicator->TestObject });
+					TestTrue(TEXT("Has Sync Control"), NewState && *NewState);
+				});
+			});
+			It("Server mute state is set", [this, RunTest]
+			{
+				RunTest();
+				
+				int32 EventCount = 0;
+				Sender->GetClientReplicationManager()
+					.QueryMuteState()
+					.Next([this, &EventCount](FConcertReplication_QueryMuteState_Response&& Response)
+					{
+						++EventCount;
+						TestEqual(TEXT("ExplicitlyMutedObjects.Num()"), Response.ExplicitlyMutedObjects.Num(), 1);
+						TestEqual(TEXT("ExplicitlyUnmutedObjects.Num()"), Response.ExplicitlyUnmutedObjects.Num(), 0);
+						TestEqual(TEXT("ImplicitlyMutedObjects.Num()"), Response.ImplicitlyMutedObjects.Num(), 0);
+						TestEqual(TEXT("ImplicitlyUnmutedObjects.Num()"), Response.ImplicitlyUnmutedObjects.Num(), 0);
+						TestTrue(TEXT("Object is muted"), Response.ExplicitlyMutedObjects.Contains(SecondObjectReplicator->TestObject));
+					});
+				TestEqual(TEXT("EventCount"), EventCount, 1);
+			});
+		});
+	}
+}
\ No newline at end of file
