Add test checking that PutState request does not produce other activities as side effect; specifically, it should not produce any mute activity, which it didn't (i.e. nothing was broken - this CL just adds a test).

diff --git a/Engine/Plugins/Developer/Concert/ConcertSync/ConcertSyncServer/Source/ConcertSyncServer/Private/Replication/ConcertServerReplicationManager_RestoreContent.cpp b/Engine/Plugins/Developer/Concert/ConcertSync/ConcertSyncServer/Source/ConcertSyncServer/Private/Replication/ConcertServerReplicationManager_RestoreContent.cpp
--- a/Engine/Plugins/Developer/Concert/ConcertSync/ConcertSyncServer/Source/ConcertSyncServer/Private/Replication/ConcertServerReplicationManager_RestoreContent.cpp
+++ b/Engine/Plugins/Developer/Concert/ConcertSync/ConcertSyncServer/Source/ConcertSyncServer/Private/Replication/ConcertServerReplicationManager_RestoreContent.cpp
@@ -402,10 +402,11 @@ namespace UE::ConcertSyncServer::Replication
 			}
 		} GroundTruth(MuteManager, ServerObjectCache);
 
-		// We'll effectively replay all mute actions that have occured so far by combining them into Request. 
+		// We'll effectively replay all mute actions that have occured so far by combining them into Request.
 		FConcertReplication_ChangeMuteState_Request AggregatedRequest;
 		ServerWorkspace.EnumerateMuteActivities([&GroundTruth, &AggregatedRequest](const FConcertSyncReplicationActivity& Activity)
 		{
+			// TODO UE-219951: We need to handle correctly replaying after PutState requests, as well! For that, we must first produce PutState activities, however (UE-219824)! 
 			FConcertSyncReplicationPayload_Mute MuteData;
 			if (!ensure(Activity.EventData.ActivityType == EConcertSyncReplicationActivityType::Mute)
 				|| !Activity.EventData.GetPayload(MuteData))
diff --git a/Engine/Plugins/Developer/Concert/ConcertSync/ConcertSyncTest/Source/ConcertSyncTest/Private/Replication/PutState/PutStateActivities.spec.cpp b/Engine/Plugins/Developer/Concert/ConcertSync/ConcertSyncTest/Source/ConcertSyncTest/Private/Replication/PutState/PutStateActivities.spec.cpp
new file mode 100644
--- /dev/null
+++ b/Engine/Plugins/Developer/Concert/ConcertSync/ConcertSyncTest/Source/ConcertSyncTest/Private/Replication/PutState/PutStateActivities.spec.cpp
@@ -0,0 +1,97 @@
+﻿// Copyright Epic Games, Inc. All Rights Reserved.
+
+#include "Replication/Messages/ChangeClientEvent.h"
+#include "Replication/Util/Spec/ObjectTestReplicator.h"
+#include "Replication/Util/Spec/ReplicationClient.h"
+#include "Replication/Util/Spec/ReplicationServer.h"
+#include "Replication/Util/ConcertClientReplicationBridgeMock.h"
+#include "Replication/ReplicationTestInterface.h"
+#include "Util/ScopedSessionDatabase.h"
+
+#include "Misc/AutomationTest.h"
+
+namespace UE::ConcertSyncTests::Replication::ChangeClients
+{
+	BEGIN_DEFINE_SPEC(FPutStateActivitiesSpec, "Editor.Concert.Replication.PutState.Activity", EAutomationTestFlags::EditorContext | EAutomationTestFlags::EngineFilter)
+		TUniquePtr<FScopedSessionDatabase> SessionDatabase;
+		TSharedPtr<ConcertSyncServer::Replication::IReplicationWorkspace> ReplicationWorkspace;
+		TUniquePtr<FObjectTestReplicator> ObjectReplicator;
+		
+		TUniquePtr<FReplicationServer> Server;
+		FReplicationClient* Client = nullptr;
+
+		const FGuid StreamId = FGuid::NewGuid();
+		/**
+		 * This is the max activity ID in the database after BeforeEach has executed.
+		 * Update this value and doc string if you make a change in the future.
+		 * For now, there are no activities produced for joining.
+		 */
+		const int32 ExpectedIdOfFirstActivityProducedByTest = 0;
+		
+		FReplicationClient& ConnectClient()
+		{
+			FReplicationClient& CreatedClient = Server->ConnectClient();
+			SessionDatabase->SetEndpoint(CreatedClient.GetEndpointId(), { CreatedClient.GetClientInfo() });
+			return CreatedClient;
+		}
+	END_DEFINE_SPEC(FPutStateActivitiesSpec);
+
+	/** This tests that the correct activities are produced when FConcertReplication_PutState_Request is used. */
+	void FPutStateActivitiesSpec::Define()
+	{
+		BeforeEach([this]
+		{
+			SessionDatabase = MakeUnique<FScopedSessionDatabase>(*this);
+			const auto FindSessionClient = [this](const FGuid& EndpointId) -> TOptional<FConcertSessionClientInfo>
+			{
+				FConcertSessionClientInfo Info;
+				const bool bFound = Server->GetServerSessionMock()->FindSessionClient(EndpointId, Info);
+				return bFound ? Info : TOptional<FConcertSessionClientInfo>{};
+			};
+			const auto ShouldIgnoreClientActivityOnRestore = [](const FGuid& EndpointId){ return false; };
+			ReplicationWorkspace = ConcertSyncServer::TestInterface::CreateReplicationWorkspace(*SessionDatabase, FindSessionClient, ShouldIgnoreClientActivityOnRestore);
+			
+			ObjectReplicator = MakeUnique<FObjectTestReplicator>();
+			Server = MakeUnique<FReplicationServer>(*this, EConcertSyncSessionFlags::Default_MultiUserSession, ReplicationWorkspace.ToSharedRef());
+			Client = &ConnectClient();
+
+			Client->JoinReplication(ObjectReplicator->CreateSenderArgs());
+		});
+		AfterEach([this]
+		{
+			Server.Reset();
+			ObjectReplicator.Reset();
+			SessionDatabase.Reset();
+			ReplicationWorkspace.Reset();
+		});
+
+		// Since FConcertReplication_PutState_Request leverages internal systems, like muting, those systems may produce an activity as side effect:
+		// validate that the put state implementation prevents that from happening. 
+		It("No unexpected activities are produced", [this]
+		{
+			FConcertReplication_PutState_Request Request;
+			Request.NewStreams.Add(Client->GetEndpointId(), { { ObjectReplicator->CreateStream(StreamId) } });
+			Request.NewAuthorityState.Add(Client->GetEndpointId(), { TArray{ FConcertObjectInStreamID{ StreamId, ObjectReplicator->TestObject } } });
+			Request.MuteChange.ObjectsToMute.Add(ObjectReplicator->TestObject);
+			
+			int32 EventCount = 0;
+			Client->GetClientReplicationManager()
+				.PutClientState(Request)
+				.Next([this, &EventCount](FConcertReplication_PutState_Response&& Response)
+				{
+					++EventCount;
+				});
+			TestEqual(TEXT("EventCount"), EventCount, 1);
+			
+			TArray<FConcertSyncReplicationActivity> ActualActivities;
+			ReplicationWorkspace->EnumerateReplicationActivities([&ActualActivities](const FConcertSyncReplicationActivity& Activity)
+			{
+				ActualActivities.Add(Activity);
+				return EBreakBehavior::Continue;
+			});
+
+			// This number might change in the future, if you change PutRequest to produce more activities yourself. Update this test accordingly in that case.
+			TestEqual(TEXT("ActualActivities.Num()"), ActualActivities.Num(), 0);
+		});
+	}
+}
