Add EConcertSyncReplicationActivityType argument to IReplicationWorkspace::GetLastReplicationActivityByClient.

diff --git a/Engine/Plugins/Developer/Concert/ConcertSync/ConcertSyncServer/Source/ConcertSyncServer/Private/Replication/ReplicationWorkspace.cpp b/Engine/Plugins/Developer/Concert/ConcertSync/ConcertSyncServer/Source/ConcertSyncServer/Private/Replication/ReplicationWorkspace.cpp
--- a/Engine/Plugins/Developer/Concert/ConcertSync/ConcertSyncServer/Source/ConcertSyncServer/Private/Replication/ReplicationWorkspace.cpp
+++ b/Engine/Plugins/Developer/Concert/ConcertSync/ConcertSyncServer/Source/ConcertSyncServer/Private/Replication/ReplicationWorkspace.cpp
@@ -36,7 +36,7 @@ namespace UE::ConcertSyncServer
 		return bSuccess ? ActivityId : TOptional<int64>{};
 	}
 
-	bool FReplicationWorkspace::GetLastReplicationActivityByClient(const FConcertSessionClientInfo& InClientInfo, FConcertSyncReplicationActivity& OutActivity) const
+	bool FReplicationWorkspace::GetLastReplicationActivityByClient(const FConcertSessionClientInfo& InClientInfo, EConcertSyncReplicationActivityType ActivityType, FConcertSyncReplicationActivity& OutActivity) const
 	{
 		// Goal: Iterate all endpoints with the same display name. Then, get the latest transaction ID from all of them.
 		// This entire approach is suboptimal though: There should be a dedicated SQL prepared statement  to which we give the client and device name.
@@ -64,7 +64,7 @@ namespace UE::ConcertSyncServer
 		for (const FConcertSyncEndpointIdAndData& ClientData : ArrayToSearch)
 		{
 			int64 EventId = INDEX_NONE;
-			Database.GetReplicationMaxEventIdByClientAndType(ClientData.EndpointId, EConcertSyncReplicationActivityType::LeaveReplication, EventId);
+			Database.GetReplicationMaxEventIdByClientAndType(ClientData.EndpointId, ActivityType, EventId);
 			NewestEventId = FMath::Max(NewestEventId, EventId);
 		}
 
diff --git a/Engine/Plugins/Developer/Concert/ConcertSync/ConcertSyncServer/Source/ConcertSyncServer/Private/Replication/ReplicationWorkspace.h b/Engine/Plugins/Developer/Concert/ConcertSync/ConcertSyncServer/Source/ConcertSyncServer/Private/Replication/ReplicationWorkspace.h
--- a/Engine/Plugins/Developer/Concert/ConcertSync/ConcertSyncServer/Source/ConcertSyncServer/Private/Replication/ReplicationWorkspace.h
+++ b/Engine/Plugins/Developer/Concert/ConcertSync/ConcertSyncServer/Source/ConcertSyncServer/Private/Replication/ReplicationWorkspace.h
@@ -38,7 +38,7 @@ namespace UE::ConcertSyncServer
 
 		//~ Begin IReplicationWorkspace Interface
 		virtual TOptional<int64> ProduceReplicationActivity(const FGuid& EndpointId, const FConcertSyncReplicationEvent& EventData) override;
-		virtual bool GetLastReplicationActivityByClient(const FConcertSessionClientInfo& InClientInfo, FConcertSyncReplicationActivity& OutActivity) const override;
+		virtual bool GetLastReplicationActivityByClient(const FConcertSessionClientInfo& InClientInfo, EConcertSyncReplicationActivityType ActivityType, FConcertSyncReplicationActivity& OutActivity) const override;
 		virtual bool GetReplicationEventById(const int64 ActivityId, FConcertSyncReplicationEvent& OutEvent) const override;
 		virtual void EnumerateReplicationActivities(TFunctionRef<EBreakBehavior(const FConcertSyncReplicationActivity& Activity)> Callback) const override;
 		//~ End IReplicationWorkspace Interface
diff --git a/Engine/Plugins/Developer/Concert/ConcertSync/ConcertSyncServer/Source/ConcertSyncServer/Public/Replication/IReplicationWorkspace.h b/Engine/Plugins/Developer/Concert/ConcertSync/ConcertSyncServer/Source/ConcertSyncServer/Public/Replication/IReplicationWorkspace.h
--- a/Engine/Plugins/Developer/Concert/ConcertSync/ConcertSyncServer/Source/ConcertSyncServer/Public/Replication/IReplicationWorkspace.h
+++ b/Engine/Plugins/Developer/Concert/ConcertSync/ConcertSyncServer/Source/ConcertSyncServer/Public/Replication/IReplicationWorkspace.h
@@ -48,7 +48,7 @@ namespace UE::ConcertSyncServer::Replication
 		 * @param OutActivity The activity, if present
 		 * @return Whether OutActivity contains a valid result.
 		 */
-		virtual bool GetLastReplicationActivityByClient(const FConcertSessionClientInfo& InClientInfo, FConcertSyncReplicationActivity& OutActivity) const = 0;
+		virtual bool GetLastReplicationActivityByClient(const FConcertSessionClientInfo& InClientInfo, EConcertSyncReplicationActivityType ActivityType, FConcertSyncReplicationActivity& OutActivity) const = 0;
 		/** Gets the last leave replication activity associated with the given client info. */
 		bool GetLastLeaveReplicationActivityByClient(const FConcertSessionClientInfo& InClientInfo, FConcertSyncReplicationPayload_LeaveReplication& OutLeaveReplication) const;
 		
@@ -79,7 +79,7 @@ namespace UE::ConcertSyncServer::Replication
 	inline bool IReplicationWorkspace::GetLastLeaveReplicationActivityByClient(const FConcertSessionClientInfo& InClientInfo, FConcertSyncReplicationPayload_LeaveReplication& OutLeaveReplication) const
 	{
 		FConcertSyncReplicationActivity Activity;
-		return GetLastReplicationActivityByClient(InClientInfo, Activity)
+		return GetLastReplicationActivityByClient(InClientInfo, EConcertSyncReplicationActivityType::LeaveReplication, Activity)
 			&& ensureMsgf(Activity.EventData.ActivityType == EConcertSyncReplicationActivityType::LeaveReplication, TEXT("Caller expected ActivityId %lld to be a LeaveReplication event"), Activity.ActivityId)
 			&& Activity.EventData.GetPayload(OutLeaveReplication);
 	}
diff --git a/Engine/Plugins/Developer/Concert/ConcertSync/ConcertSyncTest/Source/ConcertSyncTest/Private/Replication/Restore/RestoreContentPreexisting.spec.cpp b/Engine/Plugins/Developer/Concert/ConcertSync/ConcertSyncTest/Source/ConcertSyncTest/Private/Replication/Restore/RestoreContentPreexisting.spec.cpp
--- a/Engine/Plugins/Developer/Concert/ConcertSync/ConcertSyncTest/Source/ConcertSyncTest/Private/Replication/Restore/RestoreContentPreexisting.spec.cpp
+++ b/Engine/Plugins/Developer/Concert/ConcertSync/ConcertSyncTest/Source/ConcertSyncTest/Private/Replication/Restore/RestoreContentPreexisting.spec.cpp
@@ -51,7 +51,7 @@ namespace UE::ConcertSyncTests::Replication::RestoreContent
 	
 		void InsertActivityData() const
 		{
-			WorkspaceMock->ReturnResult_GetLastReplicationActivityByClient = FConcertSyncReplicationActivity(LeaveReplicationData);
+			WorkspaceMock->ReturnResult_GetLastReplicationActivityByClient = {{ EConcertSyncReplicationActivityType::LeaveReplication, FConcertSyncReplicationActivity(LeaveReplicationData) }};
 			WorkspaceMock->ReturnResult_GetReplicationEventById = FConcertSyncReplicationEvent(LeaveReplicationData);
 		}
 		FConcertReplicatedObjectInfo& GetRestoredObjectData()
diff --git a/Engine/Plugins/Developer/Concert/ConcertSync/ConcertSyncTest/Source/ConcertSyncTest/Private/Replication/Restore/RestoreContentWithNoPreexisting.spec.cpp b/Engine/Plugins/Developer/Concert/ConcertSync/ConcertSyncTest/Source/ConcertSyncTest/Private/Replication/Restore/RestoreContentWithNoPreexisting.spec.cpp
--- a/Engine/Plugins/Developer/Concert/ConcertSync/ConcertSyncTest/Source/ConcertSyncTest/Private/Replication/Restore/RestoreContentWithNoPreexisting.spec.cpp
+++ b/Engine/Plugins/Developer/Concert/ConcertSync/ConcertSyncTest/Source/ConcertSyncTest/Private/Replication/Restore/RestoreContentWithNoPreexisting.spec.cpp
@@ -33,7 +33,7 @@ namespace UE::ConcertSyncTests::Replication::RestoreContent
 
 		void InsertActivityData() const
 		{
-			WorkspaceMock->ReturnResult_GetLastReplicationActivityByClient = FConcertSyncReplicationActivity(LeaveReplicationData);
+			WorkspaceMock->ReturnResult_GetLastReplicationActivityByClient = {{ EConcertSyncReplicationActivityType::LeaveReplication, FConcertSyncReplicationActivity(LeaveReplicationData) }};
 			WorkspaceMock->ReturnResult_GetReplicationEventById = FConcertSyncReplicationEvent(LeaveReplicationData);
 		}
 
diff --git a/Engine/Plugins/Developer/Concert/ConcertSync/ConcertSyncTest/Source/ConcertSyncTest/Private/Replication/Util/Mocks/ReplicationWorkspaceCallInterceptorMock.h b/Engine/Plugins/Developer/Concert/ConcertSync/ConcertSyncTest/Source/ConcertSyncTest/Private/Replication/Util/Mocks/ReplicationWorkspaceCallInterceptorMock.h
--- a/Engine/Plugins/Developer/Concert/ConcertSync/ConcertSyncTest/Source/ConcertSyncTest/Private/Replication/Util/Mocks/ReplicationWorkspaceCallInterceptorMock.h
+++ b/Engine/Plugins/Developer/Concert/ConcertSync/ConcertSyncTest/Source/ConcertSyncTest/Private/Replication/Util/Mocks/ReplicationWorkspaceCallInterceptorMock.h
@@ -30,7 +30,7 @@ namespace UE::ConcertSyncTests::Replication
 		/** The result to return in ProduceClientLeaveReplicationActivity. */
 		TOptional<int64> ReturnResult_ProduceClientMuteReplicationActivity = 0;
 		/** The result to return in GetLastReplicationActivityByClient. */
-		TOptional<FConcertSyncReplicationActivity> ReturnResult_GetLastReplicationActivityByClient;
+		TMap<EConcertSyncReplicationActivityType, FConcertSyncReplicationActivity> ReturnResult_GetLastReplicationActivityByClient;
 		/** The result to return in GetReplicationEventById. */
 		TOptional<FConcertSyncReplicationEvent> ReturnResult_GetReplicationEventById;
 		/** The values to enumerate in EnumerateMuteActivities. */
@@ -66,14 +66,15 @@ namespace UE::ConcertSyncTests::Replication
 			}
 		}
 
-		virtual bool GetLastReplicationActivityByClient(const FConcertSessionClientInfo& InClientInfo, FConcertSyncReplicationActivity& OutActivity) const override
+		virtual bool GetLastReplicationActivityByClient(const FConcertSessionClientInfo& InClientInfo, EConcertSyncReplicationActivityType ActivityType, FConcertSyncReplicationActivity& OutActivity) const override
 		{
 			LastCall_GetLastReplicationActivityByClient = MakeTuple(InClientInfo);
-			if (ReturnResult_GetLastReplicationActivityByClient)
+			if (const FConcertSyncReplicationActivity* Result = ReturnResult_GetLastReplicationActivityByClient.Find(ActivityType))
 			{
-				OutActivity = *ReturnResult_GetLastReplicationActivityByClient;
+				OutActivity = *Result;
+				return true;
 			}
-			return ReturnResult_GetLastReplicationActivityByClient.IsSet();
+			return false;
 		}
 
 		virtual bool GetReplicationEventById(const int64 ActivityId, FConcertSyncReplicationEvent& OutEvent) const override
diff --git a/Engine/Plugins/Developer/Concert/ConcertSync/ConcertSyncTest/Source/ConcertSyncTest/Private/Replication/Util/Mocks/ReplicationWorkspaceEmptyMock.h b/Engine/Plugins/Developer/Concert/ConcertSync/ConcertSyncTest/Source/ConcertSyncTest/Private/Replication/Util/Mocks/ReplicationWorkspaceEmptyMock.h
--- a/Engine/Plugins/Developer/Concert/ConcertSync/ConcertSyncTest/Source/ConcertSyncTest/Private/Replication/Util/Mocks/ReplicationWorkspaceEmptyMock.h
+++ b/Engine/Plugins/Developer/Concert/ConcertSync/ConcertSyncTest/Source/ConcertSyncTest/Private/Replication/Util/Mocks/ReplicationWorkspaceEmptyMock.h
@@ -13,7 +13,7 @@ namespace UE::ConcertSyncTests::Replication
 		
 		//~ Begin IReplicationWorkspace Interface
 		virtual TOptional<int64> ProduceReplicationActivity(const FGuid& EndpointId, const FConcertSyncReplicationEvent& EventData) override { return {}; }
-		virtual bool GetLastReplicationActivityByClient(const FConcertSessionClientInfo& InClientInfo, FConcertSyncReplicationActivity& OutActivity) const override { return false; }
+		virtual bool GetLastReplicationActivityByClient(const FConcertSessionClientInfo& InClientInfo, EConcertSyncReplicationActivityType ActivityType, FConcertSyncReplicationActivity& OutActivity) const override { return false; }
 		virtual bool GetReplicationEventById(const int64 ActivityId, FConcertSyncReplicationEvent& OutEvent) const override { return false; }
 		virtual void EnumerateReplicationActivities(TFunctionRef<EBreakBehavior(const FConcertSyncReplicationActivity& Activity)> Callback) const override {}
 		//~ End IReplicationWorkspace Interface
