Extract NotMocked utils to MockUtils.h

diff --git a/Engine/Plugins/Developer/Concert/ConcertSync/ConcertSyncClient/Source/ConcertSyncClient/Private/ConcertClientDataStore.cpp b/Engine/Plugins/Developer/Concert/ConcertSync/ConcertSyncClient/Source/ConcertSyncClient/Private/ConcertClientDataStore.cpp
--- a/Engine/Plugins/Developer/Concert/ConcertSync/ConcertSyncClient/Source/ConcertSyncClient/Private/ConcertClientDataStore.cpp
+++ b/Engine/Plugins/Developer/Concert/ConcertSync/ConcertSyncClient/Source/ConcertSyncClient/Private/ConcertClientDataStore.cpp
@@ -239,7 +239,7 @@ int32 FConcertClientDataStore::GetCacheSize() const
 // required for normal usage. The code below provides a way to instantiate
 // a FConcertClientDataStore and look into its cache.
 //------------------------------------------------------------------------------
-namespace ConcertDataStoreTestUtils
+namespace UE::ConcertSyncTests::DatastoreTests
 {
 
 // This class is private and is meant to be indirectly used by the ConcertSyncTest module only.
diff --git a/Engine/Plugins/Developer/Concert/ConcertSync/ConcertSyncServer/Source/ConcertSyncServer/Private/ConcertServerDataStore.cpp b/Engine/Plugins/Developer/Concert/ConcertSync/ConcertSyncServer/Source/ConcertSyncServer/Private/ConcertServerDataStore.cpp
--- a/Engine/Plugins/Developer/Concert/ConcertSync/ConcertSyncServer/Source/ConcertSyncServer/Private/ConcertServerDataStore.cpp
+++ b/Engine/Plugins/Developer/Concert/ConcertSync/ConcertSyncServer/Source/ConcertSyncServer/Private/ConcertServerDataStore.cpp
@@ -136,7 +136,7 @@ void FConcertServerDataStore::FireContentReplicationEvent(const FGuid& Instigato
 	}
 }
 
-namespace ConcertDataStoreTestUtils
+namespace UE::ConcertSyncTests::DatastoreTests
 {
 
 // This function is private and declared "extern" in the ConcertSyncTest module. It enables creating a server data store to test the client/server
diff --git a/Engine/Plugins/Developer/Concert/ConcertSync/ConcertSyncTest/Source/ConcertSyncTest/Private/ConcertDataStoreTests.cpp b/Engine/Plugins/Developer/Concert/ConcertSync/ConcertSyncTest/Source/ConcertSyncTest/Private/ConcertDataStoreTests.cpp
--- a/Engine/Plugins/Developer/Concert/ConcertSync/ConcertSyncTest/Source/ConcertSyncTest/Private/ConcertDataStoreTests.cpp
+++ b/Engine/Plugins/Developer/Concert/ConcertSync/ConcertSyncTest/Source/ConcertSyncTest/Private/ConcertDataStoreTests.cpp
@@ -10,6 +10,7 @@
 #include "ConcertClientLocalDataStore.h"
 #include "IConcertSession.h"
 #include "Scratchpad/ConcertScratchpad.h"
+#include "Util/Mock/MockUtils.h"
 
 // Defines a namespace to test FText.
 #define LOCTEXT_NAMESPACE "ConcertDataStoreTests" 
@@ -17,7 +18,7 @@
 /** Flags used for the Concert data store tests. */
 static const EAutomationTestFlags ConcertDataStoreTestFlags = EAutomationTestFlags::EditorContext | EAutomationTestFlags::EngineFilter; // | EAutomationTestFlags::SmokeFilter;
 
-namespace ConcertDataStoreTestUtils
+namespace UE::ConcertSyncTests::DatastoreTests
 {
 // This function is implemented in ConcertSyncServer -> ConcertServerDataStore.cpp and prevent exposing ConcertServerDataStore.h publicly
 // as this is not required for general purpose. For the test, we pass a mocked session in which the data store server hook itself to send/
@@ -48,10 +49,6 @@ FString GetTestSessionRootPath()
 	return FPaths::ProjectIntermediateDir() / TEXT("ConcertDataStoreTest");
 }
 
-// Utility functions used to detect when a non-mocked function is called, so that we can mock it properly when required.
-template<typename T> T NotMocked(T Ret) { check(false); return Ret; }
-template<typename T> T NotMocked()      { check(false); return T(); }
-
 bool operator==(const FText& lhs, const FText& rhs)
 {
 	// This is not probably not how FText should be compared, but for the purpose of our
@@ -513,7 +510,7 @@ private:
 	TArray<TUniquePtr<FClientInfo>> Clients;
 };
 
-} // namespace ConcertDataStoreTestUtils
+} // namespace UE::ConcertSyncTests
 
 /** Ensures the Concert data store correctly versions the stored values. */
 IMPLEMENT_SIMPLE_AUTOMATION_TEST(FConcertDataStoreValueVersioning, "Editor.Concert.DataStore.ValueVersioning", ConcertDataStoreTestFlags)
@@ -596,13 +593,13 @@ bool FConcertDataStoreValueVersioning::RunTest(const FString& Parameters)
 		TestTrueExpr(Result2.Code == EConcertDataStoreResultCode::Fetched && Result2.Value->Version == 42);
 	}
 
-	IFileManager::Get().DeleteDirectory(*ConcertDataStoreTestUtils::GetTestSessionRootPath(), false, true);
+	IFileManager::Get().DeleteDirectory(*UE::ConcertSyncTests::DatastoreTests::GetTestSessionRootPath(), false, true);
 
 	return true;
 }
 
 /** Ensures the Concert data store correctly handles the common operations. */
-IMPLEMENT_CUSTOM_SIMPLE_AUTOMATION_TEST(FConcertDataStoreClientServerCommonOperations, ConcertDataStoreTestUtils::FConcertDataStoreClientServerTest, "Editor.Concert.DataStore.ClientServer.CommonOperations", ConcertDataStoreTestFlags)
+IMPLEMENT_CUSTOM_SIMPLE_AUTOMATION_TEST(FConcertDataStoreClientServerCommonOperations, UE::ConcertSyncTests::DatastoreTests::FConcertDataStoreClientServerTest, "Editor.Concert.DataStore.ClientServer.CommonOperations", ConcertDataStoreTestFlags)
 
 bool FConcertDataStoreClientServerCommonOperations::RunTest(const FString& Parameters)
 {
@@ -624,13 +621,13 @@ bool FConcertDataStoreClientServerCommonOperations::RunTest(const FString& Param
 	TestCommonOperations(Client, FName(TEXT("Key_FText")), FText(LOCTEXT("FooKey", "FooText")), FText(LOCTEXT("BarKey", "BarText")), FText(LOCTEXT("HelloKey", "HelloText")));
 	TestCommonOperations(Client, FName(TEXT("Key_Custom")), FConcertDataStore_CustomTypeTest{1, 2, 0.5f, {1}}, FConcertDataStore_CustomTypeTest{127, 8, 2.5f, {1}}, FConcertDataStore_CustomTypeTest{0, 0, 0.0f, {1}});
 
-	IFileManager::Get().DeleteDirectory(*ConcertDataStoreTestUtils::GetTestSessionRootPath(), false, true);
+	IFileManager::Get().DeleteDirectory(*UE::ConcertSyncTests::DatastoreTests::GetTestSessionRootPath(), false, true);
 
 	return true;
 }
 
 /** Ensures the Concert data store correctly handles the "key not found" cases. */
-IMPLEMENT_CUSTOM_SIMPLE_AUTOMATION_TEST(FConcertDataStoreClientServerKeyNotFound, ConcertDataStoreTestUtils::FConcertDataStoreClientServerTest, "Editor.Concert.DataStore.ClientServer.KeyNotFound", ConcertDataStoreTestFlags)
+IMPLEMENT_CUSTOM_SIMPLE_AUTOMATION_TEST(FConcertDataStoreClientServerKeyNotFound, UE::ConcertSyncTests::DatastoreTests::FConcertDataStoreClientServerTest, "Editor.Concert.DataStore.ClientServer.KeyNotFound", ConcertDataStoreTestFlags)
 
 bool FConcertDataStoreClientServerKeyNotFound::RunTest(const FString& Parameters)
 {
@@ -642,13 +639,13 @@ bool FConcertDataStoreClientServerKeyNotFound::RunTest(const FString& Parameters
 	EnsureNotFound(Client.CompareExchange(Key, 10ull, 1ull).Get());
 	EnsureNotFound(Client.CompareExchange(Key, 10.0, 1.0).Get());
 
-	IFileManager::Get().DeleteDirectory(*ConcertDataStoreTestUtils::GetTestSessionRootPath(), false, true);
+	IFileManager::Get().DeleteDirectory(*UE::ConcertSyncTests::DatastoreTests::GetTestSessionRootPath(), false, true);
 
 	return true;
 }
 
 /** Ensures the Concert data store correctly handles the "type mismatch" cases. */
-IMPLEMENT_CUSTOM_SIMPLE_AUTOMATION_TEST(FConcertDataStoreClientServerTypeMismatch, ConcertDataStoreTestUtils::FConcertDataStoreClientServerTest, "Editor.Concert.DataStore.ClientServer.TypeMismatch", ConcertDataStoreTestFlags)
+IMPLEMENT_CUSTOM_SIMPLE_AUTOMATION_TEST(FConcertDataStoreClientServerTypeMismatch, UE::ConcertSyncTests::DatastoreTests::FConcertDataStoreClientServerTest, "Editor.Concert.DataStore.ClientServer.TypeMismatch", ConcertDataStoreTestFlags)
 
 bool FConcertDataStoreClientServerTypeMismatch::RunTest(const FString& Parameters)
 {
@@ -661,13 +658,13 @@ bool FConcertDataStoreClientServerTypeMismatch::RunTest(const FString& Parameter
 	TestTypeMismatch(Client1, Client2, FName(TEXT("TypeMismatch_i64_u64")),   10ll, 1ull);
 	TestTypeMismatch(Client1, Client2, FName(TEXT("TypeMismatch_i64_bool")),  10ll, true);
 
-	IFileManager::Get().DeleteDirectory(*ConcertDataStoreTestUtils::GetTestSessionRootPath(), false, true);
+	IFileManager::Get().DeleteDirectory(*UE::ConcertSyncTests::DatastoreTests::GetTestSessionRootPath(), false, true);
 
 	return true;
 }
 
 /** Ensures the Concert data store optimize the "compare and exchange" operation, to avoid sending the payload when using the version is more optimal. */
-IMPLEMENT_CUSTOM_SIMPLE_AUTOMATION_TEST(FConcertDataStoreClientServerCompareExchangeOptimization, ConcertDataStoreTestUtils::FConcertDataStoreClientServerTest, "Editor.Concert.DataStore.ClientServer.CompareExchangeOptimization", ConcertDataStoreTestFlags)
+IMPLEMENT_CUSTOM_SIMPLE_AUTOMATION_TEST(FConcertDataStoreClientServerCompareExchangeOptimization, UE::ConcertSyncTests::DatastoreTests::FConcertDataStoreClientServerTest, "Editor.Concert.DataStore.ClientServer.CompareExchangeOptimization", ConcertDataStoreTestFlags)
 
 bool FConcertDataStoreClientServerCompareExchangeOptimization::RunTest(const FString& Parameters)
 {
@@ -679,7 +676,7 @@ bool FConcertDataStoreClientServerCompareExchangeOptimization::RunTest(const FSt
 	FConcertDataStore_CustomTypeTest ExchangeValue{0, 0, 0.0f, {0}};
 
 	// Just ensure the payload is large enough to enable the optimization.
-	TestTrueExpr(StoreValue.IntArray.Num() * sizeof(decltype(StoreValue.IntArray)::ElementType) > ConcertDataStoreTestUtils::GetCompareExchangePayloadOptimizationThreshold());
+	TestTrueExpr(StoreValue.IntArray.Num() * sizeof(decltype(StoreValue.IntArray)::ElementType) > UE::ConcertSyncTests::DatastoreTests::GetCompareExchangePayloadOptimizationThreshold());
 
 	// Add a new key. The client is expected to cache the stored value at version 1.
 	EnsureValueAdded(Client.FetchOrAdd(KeyName, StoreValue).Get(), StoreValue);
@@ -690,13 +687,13 @@ bool FConcertDataStoreClientServerCompareExchangeOptimization::RunTest(const FSt
 	// Ensure the previous exchanged value was correctly stored.
 	EnsureValueFetched(Client.FetchAs<FConcertDataStore_CustomTypeTest>(KeyName).Get(), ExchangeValue);
 
-	IFileManager::Get().DeleteDirectory(*ConcertDataStoreTestUtils::GetTestSessionRootPath(), false, true);
+	IFileManager::Get().DeleteDirectory(*UE::ConcertSyncTests::DatastoreTests::GetTestSessionRootPath(), false, true);
 
 	return true;
 }
 
 /** Ensures the Concert data store client correctly caches the key/values when it receives the response to its requests form the server. */
-IMPLEMENT_CUSTOM_SIMPLE_AUTOMATION_TEST(FConcertDataStoreClientServerClientCache, ConcertDataStoreTestUtils::FConcertDataStoreClientServerTest, "Editor.Concert.DataStore.ClientServer.ClientCache", ConcertDataStoreTestFlags)
+IMPLEMENT_CUSTOM_SIMPLE_AUTOMATION_TEST(FConcertDataStoreClientServerClientCache, UE::ConcertSyncTests::DatastoreTests::FConcertDataStoreClientServerTest, "Editor.Concert.DataStore.ClientServer.ClientCache", ConcertDataStoreTestFlags)
 
 bool FConcertDataStoreClientServerClientCache::RunTest(const FString& Parameters)
 {
@@ -706,8 +703,8 @@ bool FConcertDataStoreClientServerClientCache::RunTest(const FString& Parameters
 	IConcertClientDataStore& Client1 = ConnectClient();
 	IConcertClientDataStore& Client2 = ConnectClient();
 
-	TestTrueExpr(ConcertDataStoreTestUtils::GetClientCacheSize(Client1) == 0);
-	TestTrueExpr(ConcertDataStoreTestUtils::GetClientCacheSize(Client2) == 0);
+	TestTrueExpr(UE::ConcertSyncTests::DatastoreTests::GetClientCacheSize(Client1) == 0);
+	TestTrueExpr(UE::ConcertSyncTests::DatastoreTests::GetClientCacheSize(Client2) == 0);
 	
 	FName Key("Key");
 	int32 Value = 100;
@@ -716,13 +713,13 @@ bool FConcertDataStoreClientServerClientCache::RunTest(const FString& Parameters
 	// Ensure "not found" errors do not affect the cache.
 	EnsureNotFound(Client1.FetchAs<int64>(Key).Get());
 	EnsureNotFound(Client1.CompareExchange<int64>(Key, 0, 0).Get());
-	TestTrueExpr(ConcertDataStoreTestUtils::GetClientCacheSize(Client1) == 0);
+	TestTrueExpr(UE::ConcertSyncTests::DatastoreTests::GetClientCacheSize(Client1) == 0);
 
 	// Ensure successfully adding a value populates the client local cache.
 	EnsureValueAdded(Client1.FetchOrAdd(Key, Value).Get(), Value);
-	TestTrueExpr(ConcertDataStoreTestUtils::GetClientCacheSize(Client1) == 1);
-	TestTrueExpr(ConcertDataStoreTestUtils::GetClientCachedValue<int32>(Client1, Key)->Version == Version);
-	TestTrueExpr(ConcertDataStoreTestUtils::GetClientCachedValue<int32>(Client1, Key)->DeserializeUnchecked<int32>() == Value);
+	TestTrueExpr(UE::ConcertSyncTests::DatastoreTests::GetClientCacheSize(Client1) == 1);
+	TestTrueExpr(UE::ConcertSyncTests::DatastoreTests::GetClientCachedValue<int32>(Client1, Key)->Version == Version);
+	TestTrueExpr(UE::ConcertSyncTests::DatastoreTests::GetClientCachedValue<int32>(Client1, Key)->DeserializeUnchecked<int32>() == Value);
 	TestTrueExpr(Client1.FetchAs<int32>(Key).Get().GetValue() == Value); // Should read from the cache.
 	TestTrueExpr(Client1.FetchOrAdd<int32>(Key, Value * 2).Get().GetValue() == Value); // Should read from the cache.
 
@@ -730,25 +727,25 @@ bool FConcertDataStoreClientServerClientCache::RunTest(const FString& Parameters
 	EnsureTypeMismatch(Client1.FetchOrAdd<float>(Key, 0.0f).Get());
 	EnsureTypeMismatch(Client1.CompareExchange<float>(Key, 0.0f, 1.0f).Get());
 	EnsureTypeMismatch(Client1.FetchAs<float>(Key).Get());
-	TestTrueExpr(ConcertDataStoreTestUtils::GetClientCacheSize(Client1) == 1);
-	TestTrueExpr(ConcertDataStoreTestUtils::GetClientCachedValue<int32>(Client1, Key)->Version == Version);
-	TestTrueExpr(ConcertDataStoreTestUtils::GetClientCachedValue<int32>(Client1, Key)->DeserializeUnchecked<int32>() == Value);
+	TestTrueExpr(UE::ConcertSyncTests::DatastoreTests::GetClientCacheSize(Client1) == 1);
+	TestTrueExpr(UE::ConcertSyncTests::DatastoreTests::GetClientCachedValue<int32>(Client1, Key)->Version == Version);
+	TestTrueExpr(UE::ConcertSyncTests::DatastoreTests::GetClientCachedValue<int32>(Client1, Key)->DeserializeUnchecked<int32>() == Value);
 
 	// Ensure successfully exchanging a value updates the cache.
 	EnsureValueExchanged(Client1.CompareExchange(Key, Value, Value + 1).Get(), Value + 1);
 	++Value;
 	++Version;
-	TestTrueExpr(ConcertDataStoreTestUtils::GetClientCacheSize(Client1) == 1);
-	TestTrueExpr(ConcertDataStoreTestUtils::GetClientCachedValue<int32>(Client1, Key)->Version == Version);
-	TestTrueExpr(ConcertDataStoreTestUtils::GetClientCachedValue<int32>(Client1, Key)->DeserializeUnchecked<int32>() == Value);
+	TestTrueExpr(UE::ConcertSyncTests::DatastoreTests::GetClientCacheSize(Client1) == 1);
+	TestTrueExpr(UE::ConcertSyncTests::DatastoreTests::GetClientCachedValue<int32>(Client1, Key)->Version == Version);
+	TestTrueExpr(UE::ConcertSyncTests::DatastoreTests::GetClientCachedValue<int32>(Client1, Key)->DeserializeUnchecked<int32>() == Value);
 	TestTrueExpr(Client1.FetchAs<int32>(Key).Get().GetValue() == Value); // Should read from the cache.
 	TestTrueExpr(Client1.FetchOrAdd<int32>(Key, Value * 2).Get().GetValue() == Value); // Should read from the cache.
 
 	// Ensure failing to exchange a value does not affect the cache. (Should be local failure)
 	EnsureValueFetched(Client1.CompareExchange(Key, Value + 44, Value + 88).Get(), Value);
-	TestTrueExpr(ConcertDataStoreTestUtils::GetClientCacheSize(Client1) == 1);
-	TestTrueExpr(ConcertDataStoreTestUtils::GetClientCachedValue<int32>(Client1, Key)->Version == Version);
-	TestTrueExpr(ConcertDataStoreTestUtils::GetClientCachedValue<int32>(Client1, Key)->DeserializeUnchecked<int32>() == Value);
+	TestTrueExpr(UE::ConcertSyncTests::DatastoreTests::GetClientCacheSize(Client1) == 1);
+	TestTrueExpr(UE::ConcertSyncTests::DatastoreTests::GetClientCachedValue<int32>(Client1, Key)->Version == Version);
+	TestTrueExpr(UE::ConcertSyncTests::DatastoreTests::GetClientCachedValue<int32>(Client1, Key)->DeserializeUnchecked<int32>() == Value);
 	TestTrueExpr(Client1.FetchAs<int32>(Key).Get().GetValue() == Value); // Should read from the cache.
 	TestTrueExpr(Client1.FetchOrAdd<int32>(Key, Value * 2).Get().GetValue() == Value); // Should read from the cache.
 
@@ -758,9 +755,9 @@ bool FConcertDataStoreClientServerClientCache::RunTest(const FString& Parameters
 
 	// Force client 2 to cache key 1 by calling FetchOrAdd().
 	EnsureValueFetched(Client2.FetchOrAdd(Key, Value).Get(), Value);
-	TestTrueExpr(ConcertDataStoreTestUtils::GetClientCacheSize(Client2) == 1);
-	TestTrueExpr(ConcertDataStoreTestUtils::GetClientCachedValue<int32>(Client2, Key)->Version == Version);
-	TestTrueExpr(ConcertDataStoreTestUtils::GetClientCachedValue<int32>(Client2, Key)->DeserializeUnchecked<int32>() == Value);
+	TestTrueExpr(UE::ConcertSyncTests::DatastoreTests::GetClientCacheSize(Client2) == 1);
+	TestTrueExpr(UE::ConcertSyncTests::DatastoreTests::GetClientCachedValue<int32>(Client2, Key)->Version == Version);
+	TestTrueExpr(UE::ConcertSyncTests::DatastoreTests::GetClientCachedValue<int32>(Client2, Key)->DeserializeUnchecked<int32>() == Value);
 	TestTrueExpr(Client2.FetchAs<int32>(Key).Get().GetValue() == Value); // Should read from its cache.
 	TestTrueExpr(Client2.FetchOrAdd<int32>(Key, Value * 2).Get().GetValue() == Value); // Should read from its cache.
 
@@ -771,9 +768,9 @@ bool FConcertDataStoreClientServerClientCache::RunTest(const FString& Parameters
 	EnsureValueFetched(Client1.CompareExchange(Key, Value, Value + 10).Get(), Value + 1);
 	++Value; // This is the value as exchanged by client 2 a couple of lines above.
 	++Version;
-	TestTrueExpr(ConcertDataStoreTestUtils::GetClientCacheSize(Client1) == 1);
-	TestTrueExpr(ConcertDataStoreTestUtils::GetClientCachedValue<int32>(Client1, Key)->Version == Version);
-	TestTrueExpr(ConcertDataStoreTestUtils::GetClientCachedValue<int32>(Client1, Key)->DeserializeUnchecked<int32>() == Value);
+	TestTrueExpr(UE::ConcertSyncTests::DatastoreTests::GetClientCacheSize(Client1) == 1);
+	TestTrueExpr(UE::ConcertSyncTests::DatastoreTests::GetClientCachedValue<int32>(Client1, Key)->Version == Version);
+	TestTrueExpr(UE::ConcertSyncTests::DatastoreTests::GetClientCachedValue<int32>(Client1, Key)->DeserializeUnchecked<int32>() == Value);
 	TestTrueExpr(Client1.FetchAs<int32>(Key).Get().GetValue() == Value); // Should read from its cache.
 	TestTrueExpr(Client1.FetchOrAdd<int32>(Key, Value * 2).Get().GetValue() == Value); // Should read from its cache.
 
@@ -789,17 +786,17 @@ bool FConcertDataStoreClientServerClientCache::RunTest(const FString& Parameters
 	EnsureValueExchanged(Client1.CompareExchange<int32>(Key, Value, Value + 50).Get(), Value + 50);
 	Value += 50;
 	++Version;
-	TestTrueExpr(ConcertDataStoreTestUtils::GetClientCacheSize(Client1) == 1);
-	TestTrueExpr(ConcertDataStoreTestUtils::GetClientCachedValue<int32>(Client1, Key)->Version == Version); // Version 4.
-	TestTrueExpr(ConcertDataStoreTestUtils::GetClientCachedValue<int32>(Client1, Key)->DeserializeUnchecked<int32>() == Value);
+	TestTrueExpr(UE::ConcertSyncTests::DatastoreTests::GetClientCacheSize(Client1) == 1);
+	TestTrueExpr(UE::ConcertSyncTests::DatastoreTests::GetClientCachedValue<int32>(Client1, Key)->Version == Version); // Version 4.
+	TestTrueExpr(UE::ConcertSyncTests::DatastoreTests::GetClientCachedValue<int32>(Client1, Key)->DeserializeUnchecked<int32>() == Value);
 
-	IFileManager::Get().DeleteDirectory(*ConcertDataStoreTestUtils::GetTestSessionRootPath(), false, true);
+	IFileManager::Get().DeleteDirectory(*UE::ConcertSyncTests::DatastoreTests::GetTestSessionRootPath(), false, true);
 
 	return true;
 }
 
 /** Ensures the Concert data storeserver correctly push notifications to client and client populate its cache. */
-IMPLEMENT_CUSTOM_SIMPLE_AUTOMATION_TEST(FConcertDataStoreClientServerChangeNotification, ConcertDataStoreTestUtils::FConcertDataStoreClientServerTest, "Editor.Concert.DataStore.ClientServer.ChangeNotification", ConcertDataStoreTestFlags)
+IMPLEMENT_CUSTOM_SIMPLE_AUTOMATION_TEST(FConcertDataStoreClientServerChangeNotification, UE::ConcertSyncTests::DatastoreTests::FConcertDataStoreClientServerTest, "Editor.Concert.DataStore.ClientServer.ChangeNotification", ConcertDataStoreTestFlags)
 
 bool FConcertDataStoreClientServerChangeNotification::RunTest(const FString& Parameters)
 {
@@ -808,8 +805,8 @@ bool FConcertDataStoreClientServerChangeNotification::RunTest(const FString& Par
 	IConcertClientDataStore& Client1 = ConnectClient();
 	IConcertClientDataStore& Client2 = ConnectClient();
 
-	TestTrueExpr(ConcertDataStoreTestUtils::GetClientCacheSize(Client1) == 0);
-	TestTrueExpr(ConcertDataStoreTestUtils::GetClientCacheSize(Client2) == 0);
+	TestTrueExpr(UE::ConcertSyncTests::DatastoreTests::GetClientCacheSize(Client1) == 0);
+	TestTrueExpr(UE::ConcertSyncTests::DatastoreTests::GetClientCacheSize(Client2) == 0);
 
 	using T = int32;
 
@@ -817,18 +814,18 @@ bool FConcertDataStoreClientServerChangeNotification::RunTest(const FString& Par
 	FName Key1(TEXT("Key1"));
 	T Value = 44;
 	TestTrueExpr(Client1.FetchOrAdd(Key1, Value).Get().GetCode() == EConcertDataStoreResultCode::Added);
-	TestTrueExpr(ConcertDataStoreTestUtils::GetClientCacheSize(Client1) == 1);
+	TestTrueExpr(UE::ConcertSyncTests::DatastoreTests::GetClientCacheSize(Client1) == 1);
 
 	// Client 2 must have the key1 cached by now.
-	TestTrueExpr(ConcertDataStoreTestUtils::GetClientCacheSize(Client2) == 1);
-	TestTrueExpr(ConcertDataStoreTestUtils::GetClientCachedValue<T>(Client2, Key1)->DeserializeUnchecked<T>() == Value);
+	TestTrueExpr(UE::ConcertSyncTests::DatastoreTests::GetClientCacheSize(Client2) == 1);
+	TestTrueExpr(UE::ConcertSyncTests::DatastoreTests::GetClientCachedValue<T>(Client2, Key1)->DeserializeUnchecked<T>() == Value);
 
 	// Client 2 updates the key, server will push a notification to client 1.
 	TestTrueExpr(Client2.CompareExchange(Key1, Value, Value + 1).Get().GetCode() == EConcertDataStoreResultCode::Exchanged);
 	++Value; // The value was exchanged, set its new value.
 
 	// Client 1 cache should be updated by now.
-	TestTrueExpr(ConcertDataStoreTestUtils::GetClientCachedValue<T>(Client1, Key1)->DeserializeUnchecked<T>() == Value);
+	TestTrueExpr(UE::ConcertSyncTests::DatastoreTests::GetClientCachedValue<T>(Client1, Key1)->DeserializeUnchecked<T>() == Value);
 
 	// Add few other key/value.
 	FName Key2(TEXT("Key2"));
@@ -840,19 +837,19 @@ bool FConcertDataStoreClientServerChangeNotification::RunTest(const FString& Par
 
 	// Connect a third client. Ensure its cache gets populated by the server.
 	IConcertClientDataStore& Client3 = ConnectClient();
-	TestTrueExpr(ConcertDataStoreTestUtils::GetClientCacheSize(Client3) == 4);
-	TestTrueExpr(ConcertDataStoreTestUtils::GetClientCachedValue<T>(Client3, Key1)->DeserializeUnchecked<T>() == Value);
-	TestTrueExpr(ConcertDataStoreTestUtils::GetClientCachedValue<T>(Client3, Key2)->DeserializeUnchecked<T>() == Value);
-	TestTrueExpr(ConcertDataStoreTestUtils::GetClientCachedValue<T>(Client3, Key3)->DeserializeUnchecked<T>() == Value);
-	TestTrueExpr(ConcertDataStoreTestUtils::GetClientCachedValue<T>(Client3, Key4)->DeserializeUnchecked<T>() == Value);
+	TestTrueExpr(UE::ConcertSyncTests::DatastoreTests::GetClientCacheSize(Client3) == 4);
+	TestTrueExpr(UE::ConcertSyncTests::DatastoreTests::GetClientCachedValue<T>(Client3, Key1)->DeserializeUnchecked<T>() == Value);
+	TestTrueExpr(UE::ConcertSyncTests::DatastoreTests::GetClientCachedValue<T>(Client3, Key2)->DeserializeUnchecked<T>() == Value);
+	TestTrueExpr(UE::ConcertSyncTests::DatastoreTests::GetClientCachedValue<T>(Client3, Key3)->DeserializeUnchecked<T>() == Value);
+	TestTrueExpr(UE::ConcertSyncTests::DatastoreTests::GetClientCachedValue<T>(Client3, Key4)->DeserializeUnchecked<T>() == Value);
 
-	IFileManager::Get().DeleteDirectory(*ConcertDataStoreTestUtils::GetTestSessionRootPath(), false, true);
+	IFileManager::Get().DeleteDirectory(*UE::ConcertSyncTests::DatastoreTests::GetTestSessionRootPath(), false, true);
 
 	return true;
 }
 
 /** Ensures the Concert data store client correctly call the change handler. */
-IMPLEMENT_CUSTOM_SIMPLE_AUTOMATION_TEST(FConcertDataStoreClientServerChangeNotificationHandler, ConcertDataStoreTestUtils::FConcertDataStoreClientServerTest, "Editor.Concert.DataStore.ClientServer.ChangeNotificationHandler", ConcertDataStoreTestFlags)
+IMPLEMENT_CUSTOM_SIMPLE_AUTOMATION_TEST(FConcertDataStoreClientServerChangeNotificationHandler, UE::ConcertSyncTests::DatastoreTests::FConcertDataStoreClientServerTest, "Editor.Concert.DataStore.ClientServer.ChangeNotificationHandler", ConcertDataStoreTestFlags)
 
 bool FConcertDataStoreClientServerChangeNotificationHandler::RunTest(const FString& Parameters)
 {
@@ -965,13 +962,13 @@ bool FConcertDataStoreClientServerChangeNotificationHandler::RunTest(const FStri
 	EnsureValueExchanged(Client2.CompareExchange(DoubleKey, DoubleValue, DoubleValue + 2).Get(), DoubleValue + 2);
 	TestTrueExpr(DoubleKeyNotificationCount == 1);
 
-	IFileManager::Get().DeleteDirectory(*ConcertDataStoreTestUtils::GetTestSessionRootPath(), false, true);
+	IFileManager::Get().DeleteDirectory(*UE::ConcertSyncTests::DatastoreTests::GetTestSessionRootPath(), false, true);
 
 	return true;
 }
 
 /** Ensures the Concert data store blocking API works. */
-IMPLEMENT_CUSTOM_SIMPLE_AUTOMATION_TEST(FConcertDataStoreClientServerBlockingApi, ConcertDataStoreTestUtils::FConcertDataStoreClientServerTest, "Editor.Concert.DataStore.ClientServer.BlockingAPI", ConcertDataStoreTestFlags)
+IMPLEMENT_CUSTOM_SIMPLE_AUTOMATION_TEST(FConcertDataStoreClientServerBlockingApi, UE::ConcertSyncTests::DatastoreTests::FConcertDataStoreClientServerTest, "Editor.Concert.DataStore.ClientServer.BlockingAPI", ConcertDataStoreTestFlags)
 
 bool FConcertDataStoreClientServerBlockingApi::RunTest(const FString& Parameters)
 {
@@ -1021,13 +1018,13 @@ bool FConcertDataStoreClientServerBlockingApi::RunTest(const FString& Parameters
 	ScenerioFunc(Client1);
 	ScenerioFunc(Client2);
 
-	IFileManager::Get().DeleteDirectory(*ConcertDataStoreTestUtils::GetTestSessionRootPath(), false, true);
+	IFileManager::Get().DeleteDirectory(*UE::ConcertSyncTests::DatastoreTests::GetTestSessionRootPath(), false, true);
 
 	return true;
 }
 
 /** Ensures the Concert data store using continuation API works. */
-IMPLEMENT_CUSTOM_SIMPLE_AUTOMATION_TEST(FConcertDataStoreClientServerContinuationApi, ConcertDataStoreTestUtils::FConcertDataStoreClientServerTest, "Editor.Concert.DataStore.ClientServer.ContinuationAPI", ConcertDataStoreTestFlags)
+IMPLEMENT_CUSTOM_SIMPLE_AUTOMATION_TEST(FConcertDataStoreClientServerContinuationApi, UE::ConcertSyncTests::DatastoreTests::FConcertDataStoreClientServerTest, "Editor.Concert.DataStore.ClientServer.ContinuationAPI", ConcertDataStoreTestFlags)
 
 bool FConcertDataStoreClientServerContinuationApi::RunTest(const FString& Parameters)
 {
@@ -1077,7 +1074,7 @@ bool FConcertDataStoreClientServerContinuationApi::RunTest(const FString& Parame
 		});
 	}
 
-	IFileManager::Get().DeleteDirectory(*ConcertDataStoreTestUtils::GetTestSessionRootPath(), false, true);
+	IFileManager::Get().DeleteDirectory(*UE::ConcertSyncTests::DatastoreTests::GetTestSessionRootPath(), false, true);
 
 	// The test runs synchronously. So we expect this to be true here.
 	TestTrueExpr(bNewIdGenerated);
@@ -1204,7 +1201,7 @@ bool FConcertDataStoreClientPrivateStore::RunTest(const FString& Parameters)
 		TestTrueExpr(Result3.GetValue() == MyValue + 20);
 	}
 
-	IFileManager::Get().DeleteDirectory(*ConcertDataStoreTestUtils::GetTestSessionRootPath(), false, true);
+	IFileManager::Get().DeleteDirectory(*UE::ConcertSyncTests::DatastoreTests::GetTestSessionRootPath(), false, true);
 
 	return true;
 };
diff --git a/Engine/Plugins/Developer/Concert/ConcertSync/ConcertSyncTest/Source/ConcertSyncTest/Private/Util/ClientServerCommunicationTest.h b/Engine/Plugins/Developer/Concert/ConcertSync/ConcertSyncTest/Source/ConcertSyncTest/Private/Util/ClientServerCommunicationTest.h
--- a/Engine/Plugins/Developer/Concert/ConcertSync/ConcertSyncTest/Source/ConcertSyncTest/Private/Util/ClientServerCommunicationTest.h
+++ b/Engine/Plugins/Developer/Concert/ConcertSync/ConcertSyncTest/Source/ConcertSyncTest/Private/Util/ClientServerCommunicationTest.h
@@ -4,6 +4,7 @@
 
 #include "ConcertMessages.h"
 #include "IConcertSession.h"
+#include "Mock/MockUtils.h"
 #include "Misc/AutomationTest.h"
 #include "Misc/Paths.h"
 #include "Scratchpad/ConcertScratchpad.h"
@@ -12,10 +13,6 @@
 
 namespace UE::ConcertSyncTests
 {
-	// Utility functions used to detect when a non-mocked function is called, so that we can mock it properly when required.
-	template<typename T> T NotMocked(T Ret) { check(false); return Ret; }
-	template<typename T> T NotMocked()      { check(false); return T(); }
-
 	inline FString GetTestSessionRootPath()
 	{
 		return FPaths::ProjectIntermediateDir() / TEXT("ConcertDataStoreTest");
diff --git a/Engine/Plugins/Developer/Concert/ConcertSync/ConcertSyncTest/Source/ConcertSyncTest/Private/Util/Mock/MockUtils.h b/Engine/Plugins/Developer/Concert/ConcertSync/ConcertSyncTest/Source/ConcertSyncTest/Private/Util/Mock/MockUtils.h
new file mode 100644
--- /dev/null
+++ b/Engine/Plugins/Developer/Concert/ConcertSync/ConcertSyncTest/Source/ConcertSyncTest/Private/Util/Mock/MockUtils.h
@@ -0,0 +1,12 @@
+﻿// Copyright Epic Games, Inc. All Rights Reserved.
+
+#pragma once
+
+#include "Misc/AssertionMacros.h"
+
+namespace UE::ConcertSyncTests
+{
+	// Utility functions used to detect when a non-mocked function is called, so that we can mock it properly when required.
+	template<typename T> T NotMocked(T Ret) { check(false); return Ret; }
+	template<typename T> T NotMocked()      { check(false); return T(); }
+}
\ No newline at end of file
