Fix not being able to instantiate FConcertClientWorkspaceBaseMock due to using FConcertClientSessionBaseMock as member, which has a pure virtual.

diff --git a/Engine/Plugins/Developer/Concert/ConcertSync/ConcertSyncTest/Source/ConcertSyncTest/Private/Util/Mock/ConcertClientWorkspaceBaseMock.h b/Engine/Plugins/Developer/Concert/ConcertSync/ConcertSyncTest/Source/ConcertSyncTest/Private/Util/Mock/ConcertClientWorkspaceBaseMock.h
--- a/Engine/Plugins/Developer/Concert/ConcertSync/ConcertSyncTest/Source/ConcertSyncTest/Private/Util/Mock/ConcertClientWorkspaceBaseMock.h
+++ b/Engine/Plugins/Developer/Concert/ConcertSync/ConcertSyncTest/Source/ConcertSyncTest/Private/Util/Mock/ConcertClientWorkspaceBaseMock.h
@@ -50,7 +50,13 @@ namespace UE::ConcertSyncTests::Replication
 	private:
 		
 		// Required to get stuff to compile
-		TSharedRef<FConcertClientSessionBaseMock> SessionMock = MakeShared<FConcertClientSessionBaseMock>(FConcertClientInfo{});
+		class FSessionMock : public FConcertClientSessionBaseMock
+		{
+		public:
+			explicit FSessionMock() : FConcertClientSessionBaseMock(FConcertClientInfo{}){}
+			virtual void HandleCustomEvent(const UScriptStruct* EventType, const void* EventData) override {}
+		};
+		TSharedRef<FSessionMock> SessionMock = MakeShared<FSessionMock>();
 		FOnActivityAddedOrUpdated OnActivityAddedOrUpdatedDelegate;
 		FOnWorkspaceSynchronized OnWorkspaceSynchronizedDelegate;
 		FOnFinalizeWorkspaceSyncCompleted OnFinalizeWorkspaceSyncCompletedDelegate;
