Add API for predicting client state based on LeaveReplication activity

diff --git a/Engine/Plugins/Developer/Concert/ConcertSync/ConcertSyncClient/Source/ConcertSyncClient/Private/Replication/Misc/ClientPredictionUtils.cpp b/Engine/Plugins/Developer/Concert/ConcertSync/ConcertSyncClient/Source/ConcertSyncClient/Private/Replication/Misc/ClientPredictionUtils.cpp
new file mode 100644
--- /dev/null
+++ b/Engine/Plugins/Developer/Concert/ConcertSync/ConcertSyncClient/Source/ConcertSyncClient/Private/Replication/Misc/ClientPredictionUtils.cpp
@@ -0,0 +1,66 @@
+﻿// Copyright Epic Games, Inc. All Rights Reserved.
+
+#include "Replication/Misc/ClientPredictionUtils.h"
+
+#include "IConcertClientWorkspace.h"
+#include "Replication/Messages/ReplicationActivity.h"
+#include "Replication/Misc/StreamAndAuthorityPredictionUtils.h"
+
+namespace UE::ConcertSyncClient::Replication
+{
+	TOptional<int64> IncrementalBacktrackActivityHistoryForActivityThatSetsContent(
+		const IConcertClientWorkspace& Workspace,
+		const FConcertClientInfo& ClientInfo,
+		TArray<FConcertBaseStreamInfo>& OutStreams,
+		TArray<FConcertObjectInStreamID>& OutAuthority,
+		int64 MaxToFetch,
+		int64 MinActivityIdCutoff
+		)
+	{
+		const int64 LastActivityId = Workspace.GetLastActivityId();
+		if (MinActivityIdCutoff > LastActivityId)
+		{
+			return {};
+		}
+		
+		MaxToFetch = FMath::Max(1, MaxToFetch);
+		MinActivityIdCutoff = FMath::Max(1, MinActivityIdCutoff); // ActivityIds start at 1
+		int64 LastFirstToFetch = LastActivityId + 1;
+
+		bool bShouldStop = false;
+		// We'll walk backwards in history until we find the last state that has affect the client.
+		while (!bShouldStop) 
+		{
+			// From the back of the history take chunks of size MaxToFetch until we find what we're looking for
+			const int64 NextFirstToFetch = FMath::Max(MinActivityIdCutoff, LastFirstToFetch - MaxToFetch);
+			const int64 NumToFetch = FMath::Max(MinActivityIdCutoff, LastFirstToFetch - NextFirstToFetch);
+			LastFirstToFetch = NextFirstToFetch;
+			bShouldStop = NextFirstToFetch == MinActivityIdCutoff;
+				
+			TArray<FConcertSessionActivity> Activities;
+			TMap<FGuid, FConcertClientInfo> Endpoints;
+			Workspace.GetActivities(NextFirstToFetch, NumToFetch, Endpoints, Activities);
+
+			const auto IsEquivalentClient = [&ClientInfo, &Endpoints](const FGuid& EndpointId)
+			{
+				return ConcertSyncCore::Replication::AreLogicallySameClients(ClientInfo, Endpoints[EndpointId]);
+			};
+			const auto GetReplicationEvent = [&Workspace](const int64 EventId, TFunctionRef<void(const FConcertSyncReplicationEvent&)> Callback)
+			{
+				FConcertSyncReplicationEvent Event;
+				if (Workspace.FindReplicationEvent(EventId, Event))
+				{
+					Callback(Event);
+				}
+			};
+			const TOptional<int64> UsedActivityId = ConcertSyncCore::Replication::BacktrackActivityHistoryForActivityThatSetsContent(
+				Activities, IsEquivalentClient, GetReplicationEvent,  OutStreams, OutAuthority
+				);
+			if (UsedActivityId.IsSet())
+			{
+				return UsedActivityId;
+			}
+		}
+		return {};
+	}
+}
diff --git a/Engine/Plugins/Developer/Concert/ConcertSync/ConcertSyncClient/Source/ConcertSyncClient/Public/Replication/Misc/ClientPredictionUtils.h b/Engine/Plugins/Developer/Concert/ConcertSync/ConcertSyncClient/Source/ConcertSyncClient/Public/Replication/Misc/ClientPredictionUtils.h
new file mode 100644
--- /dev/null
+++ b/Engine/Plugins/Developer/Concert/ConcertSync/ConcertSyncClient/Source/ConcertSyncClient/Public/Replication/Misc/ClientPredictionUtils.h
@@ -0,0 +1,39 @@
+﻿// Copyright Epic Games, Inc. All Rights Reserved.
+
+#pragma once
+
+#include "Containers/ContainersFwd.h"
+#include "HAL/Platform.h"
+
+class IConcertClientWorkspace;
+
+struct FConcertBaseStreamInfo;
+struct FConcertClientInfo;
+struct FConcertObjectInStreamID;
+
+template<typename OptionalType> struct TOptional;
+
+namespace UE::ConcertSyncClient::Replication
+{
+	/**
+	 * Leverages IConcertClientWorkspace::GetActivities to get chunks from the back of the activity history and passes them to
+	 * BacktrackActivityHistoryForSettingActivities until an activity is found that sets the specified client's stream and authority content.
+	 * 
+	 * @param Workspace The workspace to retrieve the activity history chunks from 
+	 * @param ClientInfo Only activities produced by clients that are logically equivalent to this info are considered (@see AreLogicallySameClients)
+	 * @param OutStreams The streams the client had, if successful
+	 * @param OutAuthority The authority the client had, if successful
+	 * @param MaxToFetch Maximum number of activities to fetch with each IConcertClientWorkspace::GetActivities call
+	 * @param MinActivityIdCutoff This is the min activity id an activity must have to be considered. Activities before this are not considered.
+	 *
+	 * @return The activity ID that contained the state. Unset if not found.
+	 */
+	CONCERTSYNCCLIENT_API TOptional<int64> IncrementalBacktrackActivityHistoryForActivityThatSetsContent(
+		const IConcertClientWorkspace& Workspace,
+		const FConcertClientInfo& ClientInfo,
+		TArray<FConcertBaseStreamInfo>& OutStreams,
+		TArray<FConcertObjectInStreamID>& OutAuthority,
+		int64 MaxToFetch = 50,
+		int64 MinActivityIdCutoff = 1
+		);
+}
diff --git a/Engine/Plugins/Developer/Concert/ConcertSync/ConcertSyncCore/Source/ConcertSyncCore/Private/Replication/Misc/StreamAndAuthorityPredictionUtils.cpp b/Engine/Plugins/Developer/Concert/ConcertSync/ConcertSyncCore/Source/ConcertSyncCore/Private/Replication/Misc/StreamAndAuthorityPredictionUtils.cpp
--- a/Engine/Plugins/Developer/Concert/ConcertSync/ConcertSyncCore/Source/ConcertSyncCore/Private/Replication/Misc/StreamAndAuthorityPredictionUtils.cpp
+++ b/Engine/Plugins/Developer/Concert/ConcertSync/ConcertSyncCore/Source/ConcertSyncCore/Private/Replication/Misc/StreamAndAuthorityPredictionUtils.cpp
@@ -3,9 +3,109 @@
 #include "Replication/Misc/StreamAndAuthorityPredictionUtils.h"
 
 #include "ConcertMessageData.h"
+#include "Replication/Messages/ReplicationActivity.h"
+
+#include "Algo/Transform.h"
+#include "Templates/Function.h"
+
+#include <type_traits>
 
 namespace UE::ConcertSyncCore::Replication
 {
+	namespace Private
+	{
+		static bool HandleLeaveReplicationEvent(
+			const FConcertSyncReplicationEvent& ReplicationEvent,
+			TArray<FConcertBaseStreamInfo>& OutPredictedStream,
+			TArray<FConcertObjectInStreamID>& OutPredictedAuthority
+			)
+		{
+			FConcertSyncReplicationPayload_LeaveReplication LeaveReplication;
+			if (!ReplicationEvent.GetPayload(LeaveReplication))
+			{
+				return false;
+			}
+
+			Algo::Transform(LeaveReplication.Streams, OutPredictedStream, [](const FConcertReplicationStream& Stream){ return Stream.BaseDescription; });
+			OutPredictedAuthority = MoveTemp(LeaveReplication.OwnedObjects);
+			return true;
+		}
+
+		static bool AnalyzeActivity(
+			const FConcertSyncActivity& Activity,
+			TFunctionRef<bool(const FGuid& EndpointId)> IsTargetEndpointFunc,
+			FExtractReplicationEventFunc ExtractReplicationEventFunc,
+			TArray<FConcertBaseStreamInfo>& OutPredictedStream,
+			TArray<FConcertObjectInStreamID>& OutPredictedAuthority
+			)
+		{
+			if (Activity.EventType != EConcertSyncActivityEventType::Replication
+				|| !IsTargetEndpointFunc(Activity.EndpointId)
+				|| Activity.bIgnored)
+			{
+				return false;
+			}
+
+			bool bSuccess = false;
+			const uint64 EventId = Activity.EventId;
+			ExtractReplicationEventFunc(EventId, [&OutPredictedStream, &OutPredictedAuthority, &bSuccess](const FConcertSyncReplicationEvent& Event)
+			{
+				bSuccess = Event.ActivityType == EConcertSyncReplicationActivityType::LeaveReplication
+					&& HandleLeaveReplicationEvent(Event, OutPredictedStream, OutPredictedAuthority);
+			});
+			return bSuccess;
+		}
+
+		/** This indirection avoids Algo::Transform alloc overhead that would be required by the BacktrackActivityHistoryForActivityThatSetsContent overloads. */
+		template<typename TActivityType, typename TExtractLambda> requires std::is_invocable_r_v<const FConcertSyncActivity&, TExtractLambda, int32>
+		static TOptional<int64> BacktrackActivityHistoryGeneric(
+			TConstArrayView<TActivityType> Activities,
+			TFunctionRef<bool(const FGuid& EndpointId)> IsTargetEndpointFunc,
+			FExtractReplicationEventFunc GetReplicationEventFunc,
+			TArray<FConcertBaseStreamInfo>& OutStreams,
+			TArray<FConcertObjectInStreamID>& OutAuthority,
+			TExtractLambda&& ExtractLambda
+			)
+		{
+			// Walk backwards since we're looking for the latest activity that SET our client's state.
+			for (int32 i = Activities.Num() - 1; i >= 0; --i)
+			{
+				const FConcertSyncActivity& Activity = ExtractLambda(i);
+				if (AnalyzeActivity(Activity, IsTargetEndpointFunc, GetReplicationEventFunc, OutStreams, OutAuthority))
+				{
+					return Activity.ActivityId;
+				}
+			}
+			return {};
+		}
+	}
+
+	TOptional<int64> BacktrackActivityHistoryForActivityThatSetsContent(
+		TConstArrayView<FConcertSyncActivity> Activities,
+		TFunctionRef<bool(const FGuid& EndpointId)> IsTargetEndpointFunc,
+		FExtractReplicationEventFunc GetReplicationEventFunc,
+		TArray<FConcertBaseStreamInfo>& OutStreams,
+		TArray<FConcertObjectInStreamID>& OutAuthority
+		)
+	{
+		return Private::BacktrackActivityHistoryGeneric(Activities, IsTargetEndpointFunc, GetReplicationEventFunc, OutStreams, OutAuthority,
+			[&Activities](int32 i){ return Activities[i]; }
+		);
+	}
+
+	TOptional<int64> BacktrackActivityHistoryForActivityThatSetsContent(
+		TConstArrayView<FConcertSessionActivity> Activities,
+		TFunctionRef<bool(const FGuid& EndpointId)> IsTargetEndpointFunc,
+		FExtractReplicationEventFunc GetReplicationEventFunc,
+		TArray<FConcertBaseStreamInfo>& OutStreams,
+		TArray<FConcertObjectInStreamID>& OutAuthority
+		)
+	{
+		return Private::BacktrackActivityHistoryGeneric(Activities, IsTargetEndpointFunc, GetReplicationEventFunc, OutStreams, OutAuthority,
+			[&Activities](int32 i){ return Activities[i].Activity; }
+		);
+	}
+
 	bool AreLogicallySameClients(const FConcertClientInfo& First, const FConcertClientInfo& Second)
 	{
 		return First.DisplayName == Second.DisplayName && First.DeviceName == Second.DeviceName;
diff --git a/Engine/Plugins/Developer/Concert/ConcertSync/ConcertSyncCore/Source/ConcertSyncCore/Public/Replication/Misc/StreamAndAuthorityPredictionUtils.h b/Engine/Plugins/Developer/Concert/ConcertSync/ConcertSyncCore/Source/ConcertSyncCore/Public/Replication/Misc/StreamAndAuthorityPredictionUtils.h
--- a/Engine/Plugins/Developer/Concert/ConcertSync/ConcertSyncCore/Source/ConcertSyncCore/Public/Replication/Misc/StreamAndAuthorityPredictionUtils.h
+++ b/Engine/Plugins/Developer/Concert/ConcertSync/ConcertSyncCore/Source/ConcertSyncCore/Public/Replication/Misc/StreamAndAuthorityPredictionUtils.h
@@ -2,10 +2,59 @@
 
 #pragma once
 
-#include "ConcertMessageData.h"
+#include "Replication/Messages/RestoreContent.h"
+
+#include "Containers/ContainersFwd.h"
+#include "HAL/Platform.h"
+#include "Templates/FunctionFwd.h"
+
+struct FConcertBaseStreamInfo;
+struct FConcertClientInfo;
+struct FConcertObjectInStreamID;
+struct FConcertSessionActivity;
+struct FConcertSyncActivity;
+struct FConcertSyncReplicationEvent;
+
+template<typename OptionalType> struct TOptional;
 
 namespace UE::ConcertSyncCore::Replication
 {
+	using FExtractReplicationEventFunc = TFunctionRef<void(
+		const int64 EventId,
+		TFunctionRef<void(const FConcertSyncReplicationEvent& Event)> Callback
+		)>;
+	
+	/**
+	 * Walks back the activity history and finds the latest activity that sets the target client's state.
+	 *
+	 * For now, the state is only determined by EConcertSyncReplicationActivityType::LeaveReplication replication activities.
+	 * In the future, additional types may also affect it, e.g. PutState, may then each in turn produce activities.
+	 *
+	 * @param Activities The activities to analyze
+	 * @param IsTargetEndpointFunc Returns whether the endpoint ID corresponds to the client to which we want to get the state.
+	 *	Usually when restoring, you'll want to match the client's DisplayName and DeviceName to that in an activity.
+	 * @param GetReplicationEventFunc Gets the replication event for the given EventId. This is called when iterating Activities.
+	 * @param OutStreams The streams the client had, if successful
+	 * @param OutAuthority The authority the client had, if successful
+	 *
+	 * @return The activity ID that contained the state. Unset if not found.
+	 */
+	CONCERTSYNCCORE_API TOptional<int64> BacktrackActivityHistoryForActivityThatSetsContent(
+		TConstArrayView<FConcertSyncActivity> Activities,
+		TFunctionRef<bool(const FGuid& EndpointId)> IsTargetEndpointFunc,
+		FExtractReplicationEventFunc GetReplicationEventFunc,
+		TArray<FConcertBaseStreamInfo>& OutStreams,
+		TArray<FConcertObjectInStreamID>& OutAuthority
+		);
+	/** Equivalent version that accepts FConcertSessionActivity instead. This is useful if your activities come from IConcertClientWorkspace::GetActivities. */
+	CONCERTSYNCCORE_API TOptional<int64> BacktrackActivityHistoryForActivityThatSetsContent(
+		TConstArrayView<FConcertSessionActivity> Activities,
+		TFunctionRef<bool(const FGuid& EndpointId)> IsTargetEndpointFunc,
+		FExtractReplicationEventFunc GetReplicationEventFunc,
+		TArray<FConcertBaseStreamInfo>& OutStreams,
+		TArray<FConcertObjectInStreamID>& OutAuthority
+		);
+	
 	/**
 	 * Decides whether First and Second should be considered to represent the same user across several Concert sessions.
 	 *
diff --git a/Engine/Plugins/Developer/Concert/ConcertSync/ConcertSyncTest/Source/ConcertSyncTest/Private/Replication/Components/BacktrackActivityHistory.spec.cpp b/Engine/Plugins/Developer/Concert/ConcertSync/ConcertSyncTest/Source/ConcertSyncTest/Private/Replication/Components/BacktrackActivityHistory.spec.cpp
new file mode 100644
--- /dev/null
+++ b/Engine/Plugins/Developer/Concert/ConcertSync/ConcertSyncTest/Source/ConcertSyncTest/Private/Replication/Components/BacktrackActivityHistory.spec.cpp
@@ -0,0 +1,220 @@
+﻿// Copyright Epic Games, Inc. All Rights Reserved.
+
+#include "ConcertSyncSessionTypes.h"
+#include "Replication/Messages/ReplicationActivity.h"
+#include "Replication/Misc/StreamAndAuthorityPredictionUtils.h"
+#include "Replication/Util/Spec/ObjectTestReplicator.h"
+#include "Replication/Misc/ClientPredictionUtils.h"
+#include "Util/Mock/ConcertClientWorkspaceBaseMock.h"
+
+#include "Misc/AutomationTest.h"
+
+namespace UE::ConcertSyncTests::Replication::BacktrackActivityHistory
+{
+	struct FTestCase
+	{
+		const FGuid TargetClientId { 1, 0, 0, 0 };
+		const FGuid OtherClientId { 2, 0, 0, 0 };
+		const FGuid TargetStreamId { 0, 1, 0, 0 };
+		const FGuid OtherStreamId { 0, 2, 0, 0 };
+	
+		const TUniquePtr<FObjectTestReplicator> TargetReplicator = MakeUnique<FObjectTestReplicator>(TEXT("TargetObject"));
+		const TUniquePtr<FObjectTestReplicator> OtherReplicator = MakeUnique<FObjectTestReplicator>(TEXT("OtherObject"));
+
+		using EEventType = EConcertSyncActivityEventType;
+		using FActivity = FConcertSyncActivity;
+		const TArray<FConcertSyncActivity> Activities
+		{
+			FActivity{ .ActivityId = 1,						.EndpointId = OtherClientId,  .EventType = EEventType::Replication,	.EventId = 1 },
+			FActivity{ .ActivityId = 2,						.EndpointId = TargetClientId, .EventType = EEventType::Replication,	.EventId = 2 },
+			FActivity{ .ActivityId = 3, .bIgnored = true,	.EndpointId = TargetClientId, .EventType = EEventType::Replication,	.EventId = 3 },
+			FActivity{ .ActivityId = 4,						.EndpointId = TargetClientId, .EventType = EEventType::Replication,	.EventId = 4 },
+			FActivity{ .ActivityId = 5,						.EndpointId = OtherClientId,  .EventType = EEventType::Replication,	.EventId = 5 },
+			FActivity{ .ActivityId = 6,						.EndpointId = TargetClientId, .EventType = EEventType::Connection,	.EventId = 1 },
+			FActivity{ .ActivityId = 7,						.EndpointId = TargetClientId, .EventType = EEventType::Package,		.EventId = 1 },
+			FActivity{ .ActivityId = 8,						.EndpointId = TargetClientId, .EventType = EEventType::Transaction,	.EventId = 1 },
+			FActivity{ .ActivityId = 9,						.EndpointId = TargetClientId, .EventType = EEventType::Lock,		.EventId = 1 },
+		};
+	
+		const FConcertReplicationStream TargetStream = TargetReplicator->CreateStream(TargetStreamId);
+		const FConcertReplicationStream OtherStream = TargetReplicator->CreateStream(OtherStreamId);
+		const TArray<FConcertObjectInStreamID> TargetAuthority { { TargetStreamId, TargetReplicator->TestObject } };
+		const TArray<FConcertObjectInStreamID> OtherAuthority { { OtherStreamId, OtherReplicator->TestObject } };
+		const FConcertSyncReplicationPayload_LeaveReplication TargetPayload{ .Streams = { TargetStream }, .OwnedObjects = TargetAuthority };
+		const FConcertSyncReplicationPayload_LeaveReplication OtherPayload{ .Streams = { OtherStream }, .OwnedObjects = OtherAuthority };
+		
+		TOptional<FConcertSyncReplicationEvent> GetReplicationEvent(const int64 EventId, FAutomationTestBase& Test) const
+		{
+			switch (EventId)
+			{
+			case 1: Test.AddError(TEXT("Iterated from the front instead of from back")); return {};
+			case 2: return FConcertSyncReplicationEvent(TargetPayload);
+			case 3: return FConcertSyncReplicationEvent(OtherPayload);
+			case 4: return FConcertSyncReplicationEvent(FConcertSyncReplicationPayload_Mute{});
+			case 5: return FConcertSyncReplicationEvent(OtherPayload);
+			default: Test.AddError(TEXT("Unexpected EventId")); return {};
+			}
+		}
+
+		void TestContentAsExpected(
+			FAutomationTestBase& Test,
+			const TOptional<int64>& ActivityId,
+			const TArray<FConcertBaseStreamInfo>& Streams,
+			const TArray<FConcertObjectInStreamID>& Authority
+			) const
+		{
+			if (Streams.Num() != 1 || Authority.Num() != 1 || !ActivityId.IsSet())
+			{
+				Test.AddError(TEXT("Unexpected number"));
+			}
+			else
+			{
+				Test.TestEqual(TEXT("ActivityId"), *ActivityId, static_cast<int64>(2));
+				Test.TestEqual(TEXT("Streams"), Streams[0], TargetStream.BaseDescription);
+				Test.TestEqual(TEXT("Authority"), Authority, TargetAuthority);
+			}
+		}
+	};
+	
+	class FBacktrackClientWorkspaceMock : public FConcertClientWorkspaceBaseMock
+	{
+		const FTestCase& TestData;
+		const TMap<FGuid, FConcertClientInfo> Endpoints;
+		FAutomationTestBase& TestInstance;
+	public:
+		
+		FBacktrackClientWorkspaceMock(const FTestCase& TestData, TMap<FGuid, FConcertClientInfo> Endpoints, FAutomationTestBase& TestInstance)
+			: TestData(TestData)
+			, Endpoints(MoveTemp(Endpoints))
+			, TestInstance(TestInstance)
+		{}
+
+		//~ Begin IConcertClientWorkspace Interface
+		virtual int64 GetLastActivityId() const override { return 9; }
+		virtual void GetActivities(const int64 FirstActivityIdToFetch, const int64 MaxNumActivities, TMap<FGuid, FConcertClientInfo>& OutEndpointClientInfoMap, TArray<FConcertSessionActivity>& OutActivities) const override
+		{
+			// The IDs are 1 larger than their array indices
+			const int32 IdToIndex = FirstActivityIdToFetch - 1;
+			if (TestData.Activities.IsValidIndex(IdToIndex))
+			{
+				check(OutActivities.IsEmpty());
+				for (int32 i = IdToIndex; OutActivities.Num() < MaxNumActivities && TestData.Activities.IsValidIndex(i); ++i)
+				{
+					OutActivities.Emplace(TestData.Activities[i], FStructOnScope{});
+				}
+				OutEndpointClientInfoMap = Endpoints;
+			}
+		}
+		virtual bool FindReplicationEvent(const int64 ReplicationEventId, FConcertSyncReplicationEvent& OutReplicationEvent) const override
+		{
+			if (const TOptional<FConcertSyncReplicationEvent> Event = TestData.GetReplicationEvent(ReplicationEventId, TestInstance))
+			{
+				OutReplicationEvent = *Event;
+				return true;
+			}
+			return false;
+		}
+		//~ End IConcertClientWorkspace Interface
+	};
+	
+	BEGIN_DEFINE_SPEC(FBacktrackActivityHistory, "Editor.Concert.Replication.Components.BacktrackActivityHistory", EAutomationTestFlags::EditorContext | EAutomationTestFlags::EngineFilter)
+		TUniquePtr<FTestCase> TestData;
+	END_DEFINE_SPEC(FBacktrackActivityHistory);
+
+	void FBacktrackActivityHistory::Define()
+	{
+		BeforeEach([this]
+		{
+			TestData = MakeUnique<FTestCase>();
+		});
+		AfterEach([this]
+		{
+			TestData.Reset();
+		});
+		
+		It("BacktrackActivityHistoryForActivityThatSetsContent", [this]
+		{
+			TArray<FConcertBaseStreamInfo> Streams;
+			TArray<FConcertObjectInStreamID> Authority;
+			const TOptional<int64> ActivityId = ConcertSyncCore::Replication::BacktrackActivityHistoryForActivityThatSetsContent(
+				TestData->Activities,
+				[this](const FGuid& EndpointId){ return EndpointId == TestData->TargetClientId; },
+				[this](const int64 EventId, TFunctionRef<void(const FConcertSyncReplicationEvent& Event)> Callback)
+				{
+					if (const TOptional<FConcertSyncReplicationEvent> Event = TestData->GetReplicationEvent(EventId, *this))
+					{
+						Callback(*Event);
+					}
+				},
+				Streams, Authority
+				);
+
+			TestData->TestContentAsExpected(*this, ActivityId, Streams, Authority);
+		});
+
+		Describe("IncrementalBacktrackActivityHistoryForActivityThatSetsContent", [this]
+		{
+			const auto TestHasExpectedContent = [this](const int64 MaxToFetch = 50)
+			{
+				const FConcertClientInfo TargetClientInfo { .DeviceName = TEXT("TargetDevice"), .DisplayName = TEXT("TargetName") };
+				const FConcertClientInfo OtherClientInfo { .DeviceName = TEXT("OtherDevice"), .DisplayName = TEXT("OtherName") };
+				const TMap<FGuid, FConcertClientInfo> Endpoints
+				{
+					{ TestData->TargetClientId, TargetClientInfo },
+					{ TestData->OtherClientId, OtherClientInfo }
+				};
+				FBacktrackClientWorkspaceMock WorkspaceMock(*TestData, Endpoints, *this);
+				
+				TArray<FConcertBaseStreamInfo> Streams;
+				TArray<FConcertObjectInStreamID> Authority;
+				const TOptional<int64> ActivityId = ConcertSyncClient::Replication::IncrementalBacktrackActivityHistoryForActivityThatSetsContent(
+					WorkspaceMock, TargetClientInfo, Streams, Authority, MaxToFetch
+					);
+
+				TestData->TestContentAsExpected(*this, ActivityId, Streams, Authority);
+			};
+			It("MaxToFetch = 0", [this, TestHasExpectedContent]{ TestHasExpectedContent(0); });
+			It("MaxToFetch = 1", [this, TestHasExpectedContent]{ TestHasExpectedContent(1); });
+			It("MaxToFetch = 2", [this, TestHasExpectedContent]{ TestHasExpectedContent(2); });
+			It("MaxToFetch = 3", [this, TestHasExpectedContent]{ TestHasExpectedContent(3); });
+			It("MaxToFetch = 4", [this, TestHasExpectedContent]{ TestHasExpectedContent(4); });
+			It("MaxToFetch = 5", [this, TestHasExpectedContent]{ TestHasExpectedContent(5); });
+			It("MaxToFetch = 6", [this, TestHasExpectedContent]{ TestHasExpectedContent(6); });
+			It("MaxToFetch = 7", [this, TestHasExpectedContent]{ TestHasExpectedContent(7); });
+			It("MaxToFetch = 8", [this, TestHasExpectedContent]{ TestHasExpectedContent(8); });
+			It("MaxToFetch = 9", [this, TestHasExpectedContent]{ TestHasExpectedContent(9); });
+			It("MaxToFetch = 10", [this, TestHasExpectedContent]{ TestHasExpectedContent(10); });
+			It("MaxToFetch = default", [this, TestHasExpectedContent]{ TestHasExpectedContent(); });
+			
+			const auto TestDoesNotHaveContent = [this](const int64 ActivityIdCutoff = 1, const int64 MaxToFetch = 50)
+			{
+				const FConcertClientInfo TargetClientInfo { .DeviceName = TEXT("TargetDevice"), .DisplayName = TEXT("TargetName") };
+				const FConcertClientInfo OtherClientInfo { .DeviceName = TEXT("OtherDevice"), .DisplayName = TEXT("OtherName") };
+				const TMap<FGuid, FConcertClientInfo> Endpoints
+				{
+					{ TestData->TargetClientId, TargetClientInfo },
+					{ TestData->OtherClientId, OtherClientInfo }
+				};
+				FBacktrackClientWorkspaceMock WorkspaceMock(*TestData, Endpoints, *this);
+				
+				TArray<FConcertBaseStreamInfo> Streams;
+				TArray<FConcertObjectInStreamID> Authority;
+				const TOptional<int64> ActivityId = ConcertSyncClient::Replication::IncrementalBacktrackActivityHistoryForActivityThatSetsContent(
+					WorkspaceMock, TargetClientInfo, Streams, Authority, MaxToFetch, ActivityIdCutoff
+					);
+
+				TestFalse(TEXT("ActivityId.IsSet()"), ActivityId.IsSet());
+				TestEqual(TEXT("Streams.Num()"), Streams.Num(), 0);
+				TestEqual(TEXT("Authority.Num()"), Authority.Num(), 0);
+			};
+			It("ActivityIdCutoff = 3", [this, TestDoesNotHaveContent]{ TestDoesNotHaveContent(3); });
+			It("ActivityIdCutoff = 4", [this, TestDoesNotHaveContent]{ TestDoesNotHaveContent(4); });
+			It("ActivityIdCutoff = 5", [this, TestDoesNotHaveContent]{ TestDoesNotHaveContent(5); });
+			It("ActivityIdCutoff = 6", [this, TestDoesNotHaveContent]{ TestDoesNotHaveContent(6); });
+			It("ActivityIdCutoff = 7", [this, TestDoesNotHaveContent]{ TestDoesNotHaveContent(7); });
+			It("ActivityIdCutoff = 8", [this, TestDoesNotHaveContent]{ TestDoesNotHaveContent(8); });
+			It("ActivityIdCutoff = 9", [this, TestDoesNotHaveContent]{ TestDoesNotHaveContent(9); });
+			It("ActivityIdCutoff = 10", [this, TestDoesNotHaveContent]{ TestDoesNotHaveContent(10); });
+		});
+	}
+}
