Extend ObjectUtils with GetActorPathIn and ReplaceActorInPath

diff --git a/Engine/Plugins/Developer/Concert/ConcertSync/ConcertSyncCore/Source/ConcertSyncCore/Private/Misc/ObjectUtils.cpp b/Engine/Plugins/Developer/Concert/ConcertSync/ConcertSyncCore/Source/ConcertSyncCore/Private/Misc/ObjectUtils.cpp
--- a/Engine/Plugins/Developer/Concert/ConcertSync/ConcertSyncCore/Source/ConcertSyncCore/Private/Misc/ObjectUtils.cpp
+++ b/Engine/Plugins/Developer/Concert/ConcertSync/ConcertSyncCore/Source/ConcertSyncCore/Private/Misc/ObjectUtils.cpp
@@ -67,5 +67,37 @@ namespace UE::ConcertSyncCore
 		}
 		return Subpath.RightChop(LastDotIndex + 1);
 	}
+
+	TOptional<FSoftObjectPath> ReplaceActorInPath(const FSoftObjectPath& OldPath, const FSoftObjectPath& NewActor)
+	{
+		if (!IsActor(NewActor))
+		{
+			return {};
+		}
+		
+		// Example of an actor called floor
+		// SoftObjectPath = { AssetPath = {PackageName = "/Game/Maps/SyncBoxLevel", AssetName = "SyncBoxLevel"}, SubPathString = "PersistentLevel.Floor" } }
+		const FString& OldSubPathString = OldPath.GetSubPathString();
+		if (!OldSubPathString.Contains(TEXT("PersistentLevel."), ESearchCase::CaseSensitive))
+		{
+			return {};
+		}
+		const FString& NewSubPathString = NewActor.GetSubPathString();
+		
+		constexpr int32 PersistentLevelStringLength = 16; // "PersistentLevel." has 16 characters
+		constexpr int32 FirstActorCharIndex = PersistentLevelStringLength + 1;
+		const int32 IndexOfDotAfterActorName_OldSubPath = OldSubPathString.Find(TEXT("."), ESearchCase::CaseSensitive, ESearchDir::FromStart, FirstActorCharIndex);
+		
+		const bool bOldIsOnlyActor = !OldSubPathString.IsValidIndex(IndexOfDotAfterActorName_OldSubPath);
+		if (bOldIsOnlyActor)
+		{
+			return NewActor;
+		}
+
+		const FString ReplacedSubPathString = NewSubPathString
+			+ TEXT(".")
+			+ OldSubPathString.RightChop(IndexOfDotAfterActorName_OldSubPath + 1);
+		return FSoftObjectPath(NewActor.GetAssetPath(), ReplacedSubPathString);
+	}
 };
 
diff --git a/Engine/Plugins/Developer/Concert/ConcertSync/ConcertSyncCore/Source/ConcertSyncCore/Public/Misc/ObjectUtils.h b/Engine/Plugins/Developer/Concert/ConcertSync/ConcertSyncCore/Source/ConcertSyncCore/Public/Misc/ObjectUtils.h
--- a/Engine/Plugins/Developer/Concert/ConcertSync/ConcertSyncCore/Source/ConcertSyncCore/Public/Misc/ObjectUtils.h
+++ b/Engine/Plugins/Developer/Concert/ConcertSync/ConcertSyncCore/Source/ConcertSyncCore/Public/Misc/ObjectUtils.h
@@ -3,8 +3,7 @@
 #pragma once
 
 #include "Misc/Optional.h"
-
-struct FSoftObjectPath;
+#include "UObject/SoftObjectPath.h"
 
 namespace UE::ConcertSyncCore
 {
@@ -13,7 +12,32 @@ namespace UE::ConcertSyncCore
 	
 	/** @return Get the owning actor of Subobject. If Subobject is an actor, then this returns unset. */
 	CONCERTSYNCCORE_API TOptional<FSoftObjectPath> GetActorOf(const FSoftObjectPath& Subobject);
+	/** @return The actor portion of the path. Difference to GetActorOf: returns a set value whenever the path is an object in the level (actor, component, etc.) */
+	FORCEINLINE TOptional<FSoftObjectPath> GetActorPathIn(const FSoftObjectPath& Path) { return IsActor(Path) ? TOptional(Path) : GetActorOf(Path); }
 	
 	/** @return Gets the last object name in the subpath. */
 	CONCERTSYNCCORE_API FString ExtractObjectNameFromPath(const FSoftObjectPath& Object);
+
+	/**
+	 * Replaces the package path and actor name from NewActor witht that of NewActor.
+	 * 
+	 * Valid example 1:
+	 * - OldPath:		/Game/OldMap.OldMap:PersistentLevel.OldActor.Subobject
+	 * - NewActor:		/Game/NewMap.NewMap:PersistentLevel.NewActor
+	 * - Result:		/Game/NewMap.NewMap:PersistentLevel.NewActor.Subobject
+	 *
+	 * Invalid examples:
+	 * - Missing actor path:
+	 *	- OldPath:		/Game/Map.Map.Actor,	NewPath: /Game/Map.Map
+	 *	- OldPath:		/Game/Map.Map,			NewPath: /Game/Map.Map.Actor
+	 * - Subobject in NewActor are not allowed:
+	 *	- OldPath:		/Game/Map.Map:PersistentLevel.OldActor.Subobject
+	 *	- NewActor:		/Game/Map.Map:PersistentLevel.NewActor.Subobject
+	 *
+	 * @param OldPath A path to an object in a UWorld in which the actor is supposed to be replaced
+	 * @param NewActor A path to an actor
+	 *
+	 * @return The replaced object paths. Unset if the arguments were invalid.
+	 */
+	CONCERTSYNCCORE_API TOptional<FSoftObjectPath> ReplaceActorInPath(const FSoftObjectPath& OldPath, const FSoftObjectPath& NewActor);
 };
diff --git a/Engine/Plugins/Developer/Concert/ConcertSync/ConcertSyncTest/Source/ConcertSyncTest/Private/Misc/ObjectPathUtils.spec.cpp b/Engine/Plugins/Developer/Concert/ConcertSync/ConcertSyncTest/Source/ConcertSyncTest/Private/Misc/ObjectPathUtils.spec.cpp
--- a/Engine/Plugins/Developer/Concert/ConcertSync/ConcertSyncTest/Source/ConcertSyncTest/Private/Misc/ObjectPathUtils.spec.cpp
+++ b/Engine/Plugins/Developer/Concert/ConcertSync/ConcertSyncTest/Source/ConcertSyncTest/Private/Misc/ObjectPathUtils.spec.cpp
@@ -7,6 +7,7 @@
 #include "HAL/FileManager.h"
 #include "Misc/ObjectPathOuterIterator.h"
 #include "Misc/ObjectPathUtils.h"
+#include "Misc/ObjectUtils.h"
 
 namespace UE::ConcertSyncTests
 {
@@ -73,5 +74,70 @@ namespace UE::ConcertSyncTests
 			const TOptional<FSoftObjectPath> NullPath = ConcertSyncCore::GetOuterPath(nullptr);
 			TestFalse(TEXT("Null"), NullPath.IsSet());
 		});
+
+		Describe("ReplaceActorOf", [this]
+		{
+			Describe("Positive (valid args)", [this]
+			{
+				It("With subobject in OldPath", [this]
+				{
+					const FSoftObjectPath OldPath(TEXT("/Game/Map.Map:PersistentLevel.OldActor.Subobject"));
+					const FSoftObjectPath NewPath(TEXT("/Game/Map.Map:PersistentLevel.NewActor"));
+					const TOptional<FSoftObjectPath> Replacement = ConcertSyncCore::ReplaceActorInPath(OldPath, NewPath);
+					TestTrue(TEXT("Equals"), Replacement && *Replacement == FSoftObjectPath(TEXT("/Game/Map.Map:PersistentLevel.NewActor.Subobject")));
+				});
+
+				It("Only replace actor", [this]
+				{
+					const FSoftObjectPath OldSubobject(TEXT("/Game/Map.Map:PersistentLevel.Actor"));
+					const FSoftObjectPath NewPath(TEXT("/Game/Map.Map:PersistentLevel.NewActor"));
+					const TOptional<FSoftObjectPath> Replacement = ConcertSyncCore::ReplaceActorInPath(OldSubobject, NewPath);
+					TestTrue(TEXT("Equals"), Replacement && *Replacement == FSoftObjectPath(TEXT("/Game/Map.Map:PersistentLevel.NewActor")));
+				});
+
+				It("Only replace path", [this]
+				{
+					const FSoftObjectPath OldSubobject(TEXT("/Game/OldMap.OldMap:PersistentLevel.Actor"));
+					const FSoftObjectPath NewPath(TEXT("/Game/NewMap.NewMap:PersistentLevel.Actor"));
+					const TOptional<FSoftObjectPath> Replacement = ConcertSyncCore::ReplaceActorInPath(OldSubobject, NewPath);
+					TestTrue(TEXT("Equals"), Replacement && *Replacement == FSoftObjectPath(TEXT("/Game/NewMap.NewMap:PersistentLevel.Actor")));
+				});
+
+				It("Actor called \"PersistentLevel\"", [this]
+				{
+					const FSoftObjectPath OldSubobject(TEXT("/Game/OldMap.OldMap:PersistentLevel.PersistentLevel.Subobject"));
+					const FSoftObjectPath NewPath(TEXT("/Game/NewMap.NewMap:PersistentLevel.PersistentLevel"));
+					const TOptional<FSoftObjectPath> Replacement = ConcertSyncCore::ReplaceActorInPath(OldSubobject, NewPath);
+					TestTrue(TEXT("Equals"), Replacement && *Replacement == FSoftObjectPath(TEXT("/Game/NewMap.NewMap:PersistentLevel.PersistentLevel.Subobject")));
+				});
+			});
+			
+			Describe("Negative (invalid args)", [this]
+			{
+				It("OldPath has no actor", [this]
+				{
+					const FSoftObjectPath OldSubobject(TEXT("/Game/Map.Map"));
+					const FSoftObjectPath NewActor(TEXT("/Game/Map.Map:PersistentLevel.Actor.Subobject"));
+					const TOptional<FSoftObjectPath> Replacement = ConcertSyncCore::ReplaceActorInPath(OldSubobject, NewActor);
+					TestFalse(TEXT("Invalid"), Replacement.IsSet());
+				});
+
+				It("NewPath has no actor", [this]
+				{
+					const FSoftObjectPath OldSubobject(TEXT("/Game/Map.Map:PersistentLevel.Actor.Subobject"));
+					const FSoftObjectPath NewActor(TEXT("/Game/Map.Map"));
+					const TOptional<FSoftObjectPath> Replacement = ConcertSyncCore::ReplaceActorInPath(OldSubobject, NewActor);
+					TestFalse(TEXT("Invalid"), Replacement.IsSet());
+				});
+
+				It("NewPath cannot contain subobjects", [this]
+				{
+					const FSoftObjectPath OldSubobject(TEXT("/Game/Map.Map:PersistentLevel.OldActor.Subobject"));
+					const FSoftObjectPath NewPath(TEXT("/Game/Map.Map:PersistentLevel.NewActor.Subobject"));
+					const TOptional<FSoftObjectPath> Replacement = ConcertSyncCore::ReplaceActorInPath(OldSubobject, NewPath);
+					TestFalse(TEXT("Invalid"), Replacement.IsSet());
+				});
+			});
+		});
 	}
 }
