Fix level not reverting on receiving client after replication

diff --git a/Engine/Plugins/Developer/Concert/ConcertSync/ConcertSyncClient/Source/ConcertSyncClient/Private/Replication/Processing/ObjectReplicationApplierProcessor.cpp b/Engine/Plugins/Developer/Concert/ConcertSync/ConcertSyncClient/Source/ConcertSyncClient/Private/Replication/Processing/ObjectReplicationApplierProcessor.cpp
--- a/Engine/Plugins/Developer/Concert/ConcertSync/ConcertSyncClient/Source/ConcertSyncClient/Private/Replication/Processing/ObjectReplicationApplierProcessor.cpp
+++ b/Engine/Plugins/Developer/Concert/ConcertSync/ConcertSyncClient/Source/ConcertSyncClient/Private/Replication/Processing/ObjectReplicationApplierProcessor.cpp
@@ -37,6 +37,10 @@ namespace UE::ConcertSyncClient::Replication
 			CONCERT_TRACE_REPLICATION_OBJECT_SCOPE(SerializeReceivedObject, Args.ObjectInfo.ObjectId.Object, Args.ObjectInfo.SequenceId);
 			bAppliedData = true;
 
+			// We're technically modifying the package so mark it modified. This will make Concert / Multi-User revert the changes on leaving the session.
+			// This could be made faster by caching the package. For now leave it and possibly profile in future if needed.
+			Object->MarkPackageDirty();
+			
 			// TODO DP UE-193659: This is very hacky and leaves performance on the table... this is in case ApplyReplicationEvent updates the transform
 			if (USceneComponent* SceneComponent = Cast<USceneComponent>(Object))
 			{
diff --git a/Engine/Plugins/Developer/Concert/ConcertSync/ConcertSyncTest/Source/ConcertSyncTest/Private/Replication/DirtyPackageOnReplication.spec.cpp b/Engine/Plugins/Developer/Concert/ConcertSync/ConcertSyncTest/Source/ConcertSyncTest/Private/Replication/DirtyPackageOnReplication.spec.cpp
new file mode 100644
--- /dev/null
+++ b/Engine/Plugins/Developer/Concert/ConcertSync/ConcertSyncTest/Source/ConcertSyncTest/Private/Replication/DirtyPackageOnReplication.spec.cpp
@@ -0,0 +1,71 @@
+﻿// Copyright Epic Games, Inc. All Rights Reserved.
+
+#include "Misc/AutomationTest.h"
+#include "Util/ConcertClientReplicationBridgeMock.h"
+#include "Util/Spec/ObjectTestReplicator.h"
+#include "Util/Spec/ReplicationClient.h"
+#include "Util/Spec/ReplicationServer.h"
+
+namespace UE::ConcertSyncTests::Replication::UI
+{
+	BEGIN_DEFINE_SPEC(FDirtyPackageOnReplication, "Editor.Concert.Replication.DirtyPackageOnReplication", EAutomationTestFlags::EditorContext | EAutomationTestFlags::EngineFilter)
+		UPackage* Package;
+		TUniquePtr<FObjectTestReplicator> Replicator;
+		TUniquePtr<FReplicationServer> Server;
+	
+		FReplicationClient* Sender;
+		FReplicationClient* Receiver;
+
+		const FGuid StreamId = FGuid::NewGuid();
+	END_DEFINE_SPEC(FDirtyPackageOnReplication);
+
+	/**
+	 * This tests that packages are dirtied whenever anything is replicated in them.
+	 * This is important so Multi-User reverts the package back to its original state when leaving replication.
+	 */
+	void FDirtyPackageOnReplication::Define()
+	{
+		BeforeEach([this]
+		{
+			Package = NewObject<UPackage>(nullptr, *FString::Printf(TEXT("/Engine/Transient/%s"), *FGuid::NewGuid().ToString()), RF_Transient);
+			Replicator = MakeUnique<FObjectTestReplicator>(Package);
+			Server = MakeUnique<FReplicationServer>(*this);
+			Sender = &Server->ConnectClient();
+			Receiver = &Server->ConnectClient();
+
+			Sender->JoinReplication(Replicator->CreateSenderArgs(StreamId));
+			Receiver->JoinReplicationAsListener({ Replicator->TestObject });
+			Sender->GetBridgeMock().InjectAvailableObject(*Replicator->TestObject);
+			Sender->GetClientReplicationManager().TakeAuthorityOver({ Replicator->TestObject });
+
+			Package->ClearDirtyFlag();
+			TestTrue(TEXT("Test set-up correctly"), !Package->IsDirty());
+		});
+		AfterEach([this]
+		{
+			Server.Reset();
+		});
+
+		It("Package is marked dirty when replicated", [this]
+		{
+			Replicator->SimulateSendObjectToReceiver(*this, { *Sender, *Server, *Receiver }, { StreamId });
+			Replicator->TestValuesWereReplicated(*this);
+
+			TestTrue(TEXT("IsDirty()"), Package->IsDirty());
+		});
+
+		It("Package is re-marked dirty when saved in between replication", [this]
+		{
+			Replicator->SimulateSendObjectToReceiver(*this, { *Sender, *Server, *Receiver }, { StreamId });
+			Replicator->TestValuesWereReplicated(*this);
+
+			// Simulate a save
+			Package->ClearDirtyFlag();
+
+			Replicator->SimulateSendObjectToReceiver(*this, { *Sender, *Server, *Receiver }, { StreamId });
+			Replicator->TestValuesWereReplicated(*this);
+			
+			TestTrue(TEXT("IsDirty()"), Package->IsDirty());
+		});
+	}
+}
diff --git a/Engine/Plugins/Developer/Concert/ConcertSync/ConcertSyncTest/Source/ConcertSyncTest/Private/Replication/Util/Spec/ObjectTestReplicator.h b/Engine/Plugins/Developer/Concert/ConcertSync/ConcertSyncTest/Source/ConcertSyncTest/Private/Replication/Util/Spec/ObjectTestReplicator.h
--- a/Engine/Plugins/Developer/Concert/ConcertSync/ConcertSyncTest/Source/ConcertSyncTest/Private/Replication/Util/Spec/ObjectTestReplicator.h
+++ b/Engine/Plugins/Developer/Concert/ConcertSync/ConcertSyncTest/Source/ConcertSyncTest/Private/Replication/Util/Spec/ObjectTestReplicator.h
@@ -75,6 +75,7 @@ namespace UE::ConcertSyncTests::Replication
 		FObjectTestReplicator() = default;
 		explicit FObjectTestReplicator(UTestReflectionObject* TestObject) : TestObject(TestObject) {}
 		explicit FObjectTestReplicator(const FName BaseName) : TestObject(NewObject<UTestReflectionObject>(GetTransientPackage(), MakeUniqueObjectName(GetTransientPackage(), UTestReflectionObject::StaticClass(), BaseName))) {}
+		explicit FObjectTestReplicator(UPackage* Package) : TestObject(NewObject<UTestReflectionObject>(Package)) {}
 
 		/** Create a subobject of TestObject, assigns it to TestObject->InstancedSubobject, and returns a replicator for replicating it. */
 		TSharedRef<FObjectTestReplicator> CreateSubobjectReplicator(const FName BaseName = NAME_None) const;
