OnClientAuthorityServerStateChanged is no longer triggered by RestoreContent requests that do not change anything.

diff --git a/Engine/Plugins/Developer/Concert/ConcertSync/ConcertSyncClient/Source/ConcertSyncClient/Private/Replication/Manager/ReplicationManagerState_Connected.cpp b/Engine/Plugins/Developer/Concert/ConcertSync/ConcertSyncClient/Source/ConcertSyncClient/Private/Replication/Manager/ReplicationManagerState_Connected.cpp
--- a/Engine/Plugins/Developer/Concert/ConcertSync/ConcertSyncClient/Source/ConcertSyncClient/Private/Replication/Manager/ReplicationManagerState_Connected.cpp
+++ b/Engine/Plugins/Developer/Concert/ConcertSync/ConcertSyncClient/Source/ConcertSyncClient/Private/Replication/Manager/ReplicationManagerState_Connected.cpp
@@ -21,6 +21,7 @@
 #include "Misc/ScopeExit.h"
 #include "UObject/Package.h"
 #include "UObject/UObjectGlobals.h"
+#include "Utils/ContentComparisionUtils.h"
 
 namespace UE::ConcertSyncClient::Replication
 {
@@ -433,29 +434,33 @@ namespace UE::ConcertSyncClient::Replication
 
 	void FReplicationManagerState_Connected::UpdateReplicatedObjectAfterServerSideChange(const FConcertQueriedClientInfo& NewState)
 	{
-		OnPreStreamsChangedDelegate.Broadcast();
-		OnPreAuthorityChangedDelegate.Broadcast();
-		ON_SCOPE_EXIT
+		if (!AreStreamsEquivalent(NewState.Streams, RegisteredStreams))
 		{
+			OnPreStreamsChangedDelegate.Broadcast();
+			
+			RegisteredStreams.Empty(NewState.Streams.Num());
+			Algo::Transform(NewState.Streams, RegisteredStreams, [](const FConcertBaseStreamInfo& Info)
+			{
+				return FConcertReplicationStream{ Info };
+			});
+			
 			OnPostStreamsChangedDelegate.Broadcast();
-			OnPostAuthorityChangedDelegate.Broadcast();
-		};
-		
-		// The list of registered streams has changed
-		RegisteredStreams.Empty(NewState.Streams.Num());
-		Algo::Transform(NewState.Streams, RegisteredStreams, [](const FConcertBaseStreamInfo& Info)
-		{
-			return FConcertReplicationStream{ Info };
-		});
+		}
 
-		// This affects the objects we'll be replicating
-		ReplicationDataSource.ClearReplicatedObjects();
-		for (const FConcertAuthorityClientInfo& AuthorityState : NewState.Authority)
+		if (!IsAuthorityEquivalent(NewState.Authority, ReplicationDataSource))
 		{
-			for (const FSoftObjectPath& ObjectPath : AuthorityState.AuthoredObjects)
+			OnPreAuthorityChangedDelegate.Broadcast();
+			
+			ReplicationDataSource.ClearReplicatedObjects();
+			for (const FConcertAuthorityClientInfo& AuthorityState : NewState.Authority)
 			{
-				ReplicationDataSource.AddReplicatedObjectStreams(ObjectPath, { AuthorityState.StreamId });
+				for (const FSoftObjectPath& ObjectPath : AuthorityState.AuthoredObjects)
+				{
+					ReplicationDataSource.AddReplicatedObjectStreams(ObjectPath, { AuthorityState.StreamId });
+				}
 			}
+			
+			OnPostAuthorityChangedDelegate.Broadcast();
 		}
 	}
 
diff --git a/Engine/Plugins/Developer/Concert/ConcertSync/ConcertSyncClient/Source/ConcertSyncClient/Private/Replication/Manager/Utils/ContentComparisionUtils.cpp b/Engine/Plugins/Developer/Concert/ConcertSync/ConcertSyncClient/Source/ConcertSyncClient/Private/Replication/Manager/Utils/ContentComparisionUtils.cpp
new file mode 100644
--- /dev/null
+++ b/Engine/Plugins/Developer/Concert/ConcertSync/ConcertSyncClient/Source/ConcertSyncClient/Private/Replication/Manager/Utils/ContentComparisionUtils.cpp
@@ -0,0 +1,69 @@
+﻿// Copyright Epic Games, Inc. All Rights Reserved.
+
+#include "ContentComparisionUtils.h"
+
+#include "Replication/Data/ClientQueriedInfo.h"
+#include "Replication/Data/ReplicationStream.h"
+#include "Replication/Processing/ClientReplicationDataCollector.h"
+
+namespace UE::ConcertSyncClient::Replication
+{
+	bool AreStreamsEquivalent(TConstArrayView<FConcertBaseStreamInfo> NewStreams, TConstArrayView<FConcertReplicationStream> RegisteredStreams)
+	{
+		if (NewStreams.Num() != RegisteredStreams.Num())
+		{
+			return false;
+		}
+		
+		for (int32 i = 0; i < NewStreams.Num(); i++)
+		{
+			const FConcertBaseStreamInfo& NewStream = NewStreams[i];
+			const FGuid& StreamId = NewStream.Identifier;
+			const FConcertReplicationStream* StreamInfo = RegisteredStreams.FindByPredicate([&StreamId](const FConcertReplicationStream& Stream)
+			{
+				return Stream.BaseDescription.Identifier == StreamId;
+			});
+			if (!StreamInfo)
+			{
+				return false;
+			}
+
+			if (NewStream != StreamInfo->BaseDescription)
+			{
+				return false;
+			}
+		}
+		
+		return true;
+	}
+
+	bool IsAuthorityEquivalent(TConstArrayView<FConcertAuthorityClientInfo> NewAuthority, const FClientReplicationDataCollector& Replicator)
+	{
+		bool bIsEquivalent = true;
+		
+		// NewAuthority contains all of Replicator?
+		for (const FConcertAuthorityClientInfo& Authority : NewAuthority)
+		{
+			for (const FSoftObjectPath& ObjectPath : Authority.AuthoredObjects)
+			{
+				if (!Replicator.OwnsObjectInStream(ObjectPath, Authority.StreamId))
+				{
+					return false;
+				}
+			}
+		}
+		
+		// Replicator contains all of NewAuthority?
+		Replicator.ForEachOwnedObjectAndStream([&NewAuthority, &bIsEquivalent](const FSoftObjectPath& ReplicatedObject, const FGuid& StreamId)
+		{
+			const FConcertAuthorityClientInfo* ClientInfo = NewAuthority.FindByPredicate([&StreamId](const FConcertAuthorityClientInfo& Info)
+			{
+				return Info.StreamId == StreamId;
+			});
+			bIsEquivalent &= ClientInfo && ClientInfo->AuthoredObjects.Contains(ReplicatedObject);
+			return bIsEquivalent ? EBreakBehavior::Continue : EBreakBehavior::Break;
+		});
+		
+		return bIsEquivalent;
+	}
+}
diff --git a/Engine/Plugins/Developer/Concert/ConcertSync/ConcertSyncClient/Source/ConcertSyncClient/Private/Replication/Manager/Utils/ContentComparisionUtils.h b/Engine/Plugins/Developer/Concert/ConcertSync/ConcertSyncClient/Source/ConcertSyncClient/Private/Replication/Manager/Utils/ContentComparisionUtils.h
new file mode 100644
--- /dev/null
+++ b/Engine/Plugins/Developer/Concert/ConcertSync/ConcertSyncClient/Source/ConcertSyncClient/Private/Replication/Manager/Utils/ContentComparisionUtils.h
@@ -0,0 +1,20 @@
+﻿// Copyright Epic Games, Inc. All Rights Reserved.
+
+#pragma once
+
+#include "Containers/ContainersFwd.h"
+
+struct FConcertAuthorityClientInfo;
+struct FConcertBaseStreamInfo;
+struct FConcertReplicationStream;
+
+namespace UE::ConcertSyncClient::Replication
+{
+	class FClientReplicationDataCollector;
+	
+	/** @return Whether NewStreams would change nothing about RegisteredStreams if replaced. */
+	bool AreStreamsEquivalent(TConstArrayView<FConcertBaseStreamInfo> NewStreams, TConstArrayView<FConcertReplicationStream> RegisteredStreams);
+
+	/** @return Whether assigning NewAuthority to Replicator would change any authority. */
+	bool IsAuthorityEquivalent(TConstArrayView<FConcertAuthorityClientInfo> NewAuthority, const FClientReplicationDataCollector& Replicator);
+}
diff --git a/Engine/Plugins/Developer/Concert/ConcertSync/ConcertSyncClient/Source/ConcertSyncClient/Private/Replication/Processing/ClientReplicationDataCollector.cpp b/Engine/Plugins/Developer/Concert/ConcertSync/ConcertSyncClient/Source/ConcertSyncClient/Private/Replication/Processing/ClientReplicationDataCollector.cpp
--- a/Engine/Plugins/Developer/Concert/ConcertSync/ConcertSyncClient/Source/ConcertSyncClient/Private/Replication/Processing/ClientReplicationDataCollector.cpp
+++ b/Engine/Plugins/Developer/Concert/ConcertSync/ConcertSyncClient/Source/ConcertSyncClient/Private/Replication/Processing/ClientReplicationDataCollector.cpp
@@ -140,19 +140,6 @@ namespace UE::ConcertSyncClient::Replication
 		ObjectsToReplicate.Reset();
 	}
 
-	void FClientReplicationDataCollector::ForEachOwnedObject(
-		TFunctionRef<EBreakBehavior(const FSoftObjectPath&)> Callback
-		) const
-	{
-		for (const TPair<FSoftObjectPath, TArray<FObjectInfo>>& Pair : ObjectsToReplicate)
-		{
-			if (Callback(Pair.Key) == EBreakBehavior::Break)
-			{
-				break;
-			}
-		}
-	}
-
 	void FClientReplicationDataCollector::AppendOwningStreamsForObject(
 		const FSoftObjectPath& ObjectPath,
 		TSet<FGuid>& Paths
@@ -167,6 +154,12 @@ namespace UE::ConcertSyncClient::Replication
 		}
 	}
 
+	bool FClientReplicationDataCollector::OwnsObjectInStream(const FSoftObjectPath& ObjectPath, const FGuid& StreamId) const
+	{
+		const TArray<FObjectInfo>* ObjectInfo = ObjectsToReplicate.Find(ObjectPath);
+		return ObjectInfo && ObjectInfo->ContainsByPredicate([&StreamId](const FObjectInfo& ObjectInfo){ return ObjectInfo.StreamId == StreamId; });
+	}
+
 	void FClientReplicationDataCollector::ForEachPendingObject(TFunctionRef<void(const ConcertSyncCore::FPendingObjectReplicationInfo&)> ProcessItemFunc) const
 	{
 		for (const TPair<FSoftObjectPath, TArray<FObjectInfo>>& Pair : ObjectsToReplicate)
diff --git a/Engine/Plugins/Developer/Concert/ConcertSync/ConcertSyncClient/Source/ConcertSyncClient/Private/Replication/Processing/ClientReplicationDataCollector.h b/Engine/Plugins/Developer/Concert/ConcertSync/ConcertSyncClient/Source/ConcertSyncClient/Private/Replication/Processing/ClientReplicationDataCollector.h
--- a/Engine/Plugins/Developer/Concert/ConcertSync/ConcertSyncClient/Source/ConcertSyncClient/Private/Replication/Processing/ClientReplicationDataCollector.h
+++ b/Engine/Plugins/Developer/Concert/ConcertSync/ConcertSyncClient/Source/ConcertSyncClient/Private/Replication/Processing/ClientReplicationDataCollector.h
@@ -70,10 +70,17 @@ namespace UE::ConcertSyncClient::Replication
 		void ClearReplicatedObjects();
 
 		/** Iterates every object for which there is at least one owning stream. */
-		void ForEachOwnedObject(TFunctionRef<EBreakBehavior(const FSoftObjectPath&)> Callback) const;
+		template<typename TCallback> requires std::is_invocable_r_v<EBreakBehavior, TCallback, const FSoftObjectPath&>
+		void ForEachOwnedObject(TCallback&& Callback) const;
+		/** Iterates every object and stream assigned to it. */
+		template<typename TCallback> requires std::is_invocable_r_v<EBreakBehavior, TCallback, const FSoftObjectPath&, const FGuid&>
+		void ForEachOwnedObjectAndStream(TCallback&& Callback) const;
 		/** Writes all owning streams for ObjectPath into Paths. */
 		void AppendOwningStreamsForObject(const FSoftObjectPath& ObjectPath, TSet<FGuid>& Paths) const;
 
+		/** @return Whether ObjectPath is owned with StreamId. */
+		bool OwnsObjectInStream(const FSoftObjectPath& ObjectPath, const FGuid& StreamId) const;
+
 		//~ Begin IReplicationDataSource Interface
 		virtual void ForEachPendingObject(TFunctionRef<void(const ConcertSyncCore::FPendingObjectReplicationInfo&)> ProcessItemFunc) const override;
 		virtual int32 NumObjects() const override { return NumTrackedObjects; }
@@ -116,4 +123,31 @@ namespace UE::ConcertSyncClient::Replication
 		void StartTrackingObject(UObject& Object);
 		void StopTrackingObject(const FSoftObjectPath& ObjectPath);
 	};
+
+	template <typename TCallback> requires std::is_invocable_r_v<EBreakBehavior, TCallback, const FSoftObjectPath&>
+	void FClientReplicationDataCollector::ForEachOwnedObject(TCallback&& Callback) const
+	{
+		for (const TPair<FSoftObjectPath, TArray<FObjectInfo>>& Pair : ObjectsToReplicate)
+		{
+			if (Callback(Pair.Key) == EBreakBehavior::Break)
+			{
+				break;
+			}
+		}
+	}
+
+	template <typename TCallback> requires std::is_invocable_r_v<EBreakBehavior, TCallback, const FSoftObjectPath&, const FGuid&>
+	void FClientReplicationDataCollector::ForEachOwnedObjectAndStream(TCallback&& Callback) const
+	{
+		for (const TPair<FSoftObjectPath, TArray<FObjectInfo>>& Pair : ObjectsToReplicate)
+		{
+			for (const FObjectInfo& ObjectInfo : Pair.Value)
+			{
+				if (Callback(Pair.Key, ObjectInfo.StreamId) == EBreakBehavior::Break)
+				{
+					return;
+				}
+			}
+		}
+	}
 }
diff --git a/Engine/Plugins/Developer/Concert/ConcertSync/ConcertSyncTest/Source/ConcertSyncTest/Private/Replication/ClientManager/RestoreContentEvents.spec.cpp b/Engine/Plugins/Developer/Concert/ConcertSync/ConcertSyncTest/Source/ConcertSyncTest/Private/Replication/ClientManager/RestoreContentEvents.spec.cpp
new file mode 100644
--- /dev/null
+++ b/Engine/Plugins/Developer/Concert/ConcertSync/ConcertSyncTest/Source/ConcertSyncTest/Private/Replication/ClientManager/RestoreContentEvents.spec.cpp
@@ -0,0 +1,126 @@
+﻿// Copyright Epic Games, Inc. All Rights Reserved.
+
+#include "Misc/AutomationTest.h"
+
+#include "Replication/Util/Spec/ClientEventCounter.h"
+#include "Replication/Util/Spec/ObjectTestReplicator.h"
+#include "Replication/Util/Spec/ReplicationClient.h"
+#include "Replication/Util/Spec/ReplicationServer.h"
+
+namespace UE::ConcertSyncTests::Replication::UI
+{
+	static FConcertReplication_RestoreContent_Response CreateRestoreContentFrom(
+		const FObjectTestReplicator& Replication,
+		bool bGiveAuthority = true,
+		const FGuid& StreamId = FGuid::NewGuid()
+		)
+	{
+		FConcertReplication_RestoreContent_Response Result;
+
+		Result.ErrorCode = EConcertReplicationRestoreErrorCode::Success;
+		Result.ClientInfo.Streams = { Replication.CreateStream(StreamId).BaseDescription };
+		if (bGiveAuthority)
+		{
+			Result.ClientInfo.Authority = { FConcertAuthorityClientInfo{ StreamId, { Replication.TestObject } } };
+		}
+		
+		return Result;
+	}
+	
+	BEGIN_DEFINE_SPEC(FRestoreContentEvents, "Editor.Concert.Replication.ClientManager.RestoreContentEvents", EAutomationTestFlags::EditorContext | EAutomationTestFlags::EngineFilter)
+		TUniquePtr<FObjectTestReplicator> Replicator;
+		TUniquePtr<FReplicationServer> Server;
+		FReplicationClient* Sender;
+
+		/** Will help us count the client events. */
+		FClientEventCounter EventCounter;
+
+		FConcertReplication_RestoreContent_Response ResponseToGive;
+		EConcertSessionResponseCode HandleRestoreContentRequest(
+			const FConcertSessionContext& Context,
+			const FConcertReplication_RestoreContent_Request& Request,
+			FConcertReplication_RestoreContent_Response& Response
+			)
+		{
+			Response = ResponseToGive;
+			return EConcertSessionResponseCode::Success;
+		}
+	END_DEFINE_SPEC(FRestoreContentEvents);
+
+	/**
+	 * This tests that IConcertClientReplicationManager stream and authority change events are made correctly in response to
+	 * IConcertClientReplicationManager::RestoreContent.
+	 */
+	void FRestoreContentEvents::Define()
+	{
+		BeforeEach([this]
+		{
+			Replicator = MakeUnique<FObjectTestReplicator>();
+			Server = MakeUnique<FReplicationServer>(*this);
+			Sender = &Server->ConnectClient();
+
+			Server->GetServerSessionMock()->RegisterCustomRequestHandler<FConcertReplication_RestoreContent_Request, FConcertReplication_RestoreContent_Response>(
+				this, &FRestoreContentEvents::HandleRestoreContentRequest
+				);
+			ResponseToGive = { EConcertReplicationRestoreErrorCode::Success };
+
+			Sender->JoinReplication();
+			EventCounter.Subscribe(*Sender);
+		});
+		AfterEach([this]
+		{
+			Replicator.Reset();
+			Server.Reset();
+		});
+
+		It("When content is restored, events called", [this]
+		{
+			ResponseToGive = CreateRestoreContentFrom(*Replicator);
+			Sender->GetClientReplicationManager().RestoreContent();
+			EventCounter.TestCount(*this, 1, 1);
+		});
+		It("When only stream changes, only stream change event is called", [this]
+		{
+			constexpr bool bGiveAuthority = false;
+			ResponseToGive = CreateRestoreContentFrom(*Replicator, bGiveAuthority);
+			Sender->GetClientReplicationManager().RestoreContent();
+			EventCounter.TestCount(*this, 1, 0);
+		});
+		It("When only authority changes, only authority change event is called", [this]
+		{
+			const FGuid& StreamId = FGuid::NewGuid();
+			
+			ResponseToGive = CreateRestoreContentFrom(*Replicator, true, StreamId);
+			Sender->GetClientReplicationManager().RestoreContent();
+			EventCounter.ResetEventCount();
+			
+			ResponseToGive = CreateRestoreContentFrom(*Replicator, false, StreamId);
+			Sender->GetClientReplicationManager().RestoreContent();
+			EventCounter.TestCount(*this, 0, 1);
+		});
+
+		It("When empty content is restored, no events called", [this]
+		{
+			ResponseToGive = { EConcertReplicationRestoreErrorCode::Success };
+			Sender->GetClientReplicationManager().RestoreContent();
+			EventCounter.TestCount(*this, 0, 0);
+		});
+		It("When same content is restored again, no events called", [this]
+		{
+			ResponseToGive = CreateRestoreContentFrom(*Replicator);
+			Sender->GetClientReplicationManager().RestoreContent();
+			EventCounter.ResetEventCount();
+			
+			// This should generate no events because it's the same content.
+			Sender->GetClientReplicationManager().RestoreContent();
+			EventCounter.TestCount(*this, 0, 0);
+		});
+
+		It("When RestoreContent times out, no events called ", [this]
+		{
+			ResponseToGive = { EConcertReplicationRestoreErrorCode::Timeout };
+			Sender->GetClientReplicationManager().RestoreContent();
+			EventCounter.TestCount(*this, 0, 0);
+		});
+	}
+}
\ No newline at end of file
diff --git a/Engine/Plugins/Developer/Concert/ConcertSync/ConcertSyncTest/Source/ConcertSyncTest/Private/Replication/Util/Spec/ClientEventCounter.h b/Engine/Plugins/Developer/Concert/ConcertSync/ConcertSyncTest/Source/ConcertSyncTest/Private/Replication/Util/Spec/ClientEventCounter.h
new file mode 100644
--- /dev/null
+++ b/Engine/Plugins/Developer/Concert/ConcertSync/ConcertSyncTest/Source/ConcertSyncTest/Private/Replication/Util/Spec/ClientEventCounter.h
@@ -0,0 +1,42 @@
+﻿// Copyright Epic Games, Inc. All Rights Reserved.
+
+#pragma once
+
+#include "ReplicationClient.h"
+
+namespace UE::ConcertSyncTests::Replication
+{
+	/** Helps track the number of authority and stream change events in IConcertClientReplicationManager.  */
+	struct FClientEventCounter
+	{
+		int32 CallCount_PreAuthorityChanged = 0;
+		int32 CallCount_PostAuthorityChanged = 0;
+		int32 CallCount_PreStreamsChanged = 0;
+		int32 CallCount_PostStreamsChanged = 0;
+		
+		void ResetEventCount()
+		{
+			CallCount_PreAuthorityChanged = 0;
+			CallCount_PostAuthorityChanged = 0;
+			CallCount_PreStreamsChanged = 0;
+			CallCount_PostStreamsChanged = 0;
+		}
+		
+		void TestCount(FAutomationTestBase& Test, int32 ExpectedStreamCount, int32 ExpectedAuthorityCount) const
+		{
+			Test.TestEqual(TEXT("PreAuthorityChanged"), CallCount_PreAuthorityChanged, ExpectedAuthorityCount);
+			Test.TestEqual(TEXT("PostAuthorityChanged"), CallCount_PostAuthorityChanged, ExpectedAuthorityCount);
+			Test.TestEqual(TEXT("PreStreamsChanged"), CallCount_PreStreamsChanged, ExpectedStreamCount);
+			Test.TestEqual(TEXT("PostStreamsChanged"), CallCount_PostStreamsChanged, ExpectedStreamCount);
+		}
+
+		void Subscribe(const FReplicationClient& Client)
+		{
+			Client.GetClientReplicationManager().OnPreAuthorityChanged().AddLambda([this]{ ++CallCount_PreAuthorityChanged; });
+			Client.GetClientReplicationManager().OnPostAuthorityChanged().AddLambda([this]{ ++CallCount_PostAuthorityChanged; });
+			Client.GetClientReplicationManager().OnPreStreamsChanged().AddLambda([this]{ ++CallCount_PreStreamsChanged; });
+			Client.GetClientReplicationManager().OnPostStreamsChanged().AddLambda([this]{ ++CallCount_PostStreamsChanged; });
+			ResetEventCount();
+		}
+	};
+}
