Automation - Moving Concert unit tests from Editor to VirtualProduction test suite

diff --git a/Engine/Plugins/Developer/Concert/ConcertSync/ConcertSyncClient/Source/ConcertSyncClient/Test/ConcertClientLiveTransactionAuthorsTests.cpp b/Engine/Plugins/Developer/Concert/ConcertSync/ConcertSyncClient/Source/ConcertSyncClient/Test/ConcertClientLiveTransactionAuthorsTests.cpp
--- a/Engine/Plugins/Developer/Concert/ConcertSync/ConcertSyncClient/Source/ConcertSyncClient/Test/ConcertClientLiveTransactionAuthorsTests.cpp
+++ b/Engine/Plugins/Developer/Concert/ConcertSync/ConcertSyncClient/Source/ConcertSyncClient/Test/ConcertClientLiveTransactionAuthorsTests.cpp
@@ -157,7 +157,7 @@ void FenceTransactions(FConcertSyncClientLiveSession& InLiveSession, FConcertCli
 }
 
 /** Ensures the live transaction authors works correctly when there is no other clients connected. */
-IMPLEMENT_SIMPLE_AUTOMATION_TEST(FConcertLiveTransactionAuthorsSingleClient, "Editor.Concert.LiveTransactionAuthors.SingleClient", ConcertClientLiveTransactionAuthorsTestFlags)
+IMPLEMENT_SIMPLE_AUTOMATION_TEST(FConcertLiveTransactionAuthorsSingleClient, "VirtualProduction.Concert.LiveTransactionAuthors.SingleClient", ConcertClientLiveTransactionAuthorsTestFlags)
 
 bool FConcertLiveTransactionAuthorsSingleClient::RunTest(const FString& Parameters)
 {
@@ -196,7 +196,7 @@ bool FConcertLiveTransactionAuthorsSingleClient::RunTest(const FString& Paramete
 }
 
 /** Ensures the live transaction authors works correctly when there are many clients connected. */
-IMPLEMENT_SIMPLE_AUTOMATION_TEST(FConcertLiveTransactionAuthorsManyClients, "Editor.Concert.LiveTransactionAuthors.ManyClients", ConcertClientLiveTransactionAuthorsTestFlags)
+IMPLEMENT_SIMPLE_AUTOMATION_TEST(FConcertLiveTransactionAuthorsManyClients, "VirtualProduction.Concert.LiveTransactionAuthors.ManyClients", ConcertClientLiveTransactionAuthorsTestFlags)
 
 bool FConcertLiveTransactionAuthorsManyClients::RunTest(const FString& Parameters)
 {
@@ -310,7 +310,7 @@ bool FConcertLiveTransactionAuthorsManyClients::RunTest(const FString& Parameter
 }
 
 /** Ensures the live transaction authors works correctly when they are some transaction owned by a disconnected client. */
-IMPLEMENT_SIMPLE_AUTOMATION_TEST(FConcertLiveTransactionAuthorsDisconnectedClient, "Editor.Concert.LiveTransactionAuthors.DisconnectedClient", ConcertClientLiveTransactionAuthorsTestFlags)
+IMPLEMENT_SIMPLE_AUTOMATION_TEST(FConcertLiveTransactionAuthorsDisconnectedClient, "VirtualProduction.Concert.LiveTransactionAuthors.DisconnectedClient", ConcertClientLiveTransactionAuthorsTestFlags)
 
 bool FConcertLiveTransactionAuthorsDisconnectedClient::RunTest(const FString& Parameters)
 {
@@ -365,7 +365,7 @@ bool FConcertLiveTransactionAuthorsDisconnectedClient::RunTest(const FString& Pa
 }
 
 /** Ensures the live transaction authors works correctly when the same person is editing a package from two editors, on the same machine, concurrently. */
-IMPLEMENT_SIMPLE_AUTOMATION_TEST(FConcertLiveTransactionAuthorsClientUsingTwoEditors, "Editor.Concert.LiveTransactionAuthors.SamePersonUsingTwoEditorsConcurrently", ConcertClientLiveTransactionAuthorsTestFlags)
+IMPLEMENT_SIMPLE_AUTOMATION_TEST(FConcertLiveTransactionAuthorsClientUsingTwoEditors, "VirtualProduction.Concert.LiveTransactionAuthors.SamePersonUsingTwoEditorsConcurrently", ConcertClientLiveTransactionAuthorsTestFlags)
 
 bool FConcertLiveTransactionAuthorsClientUsingTwoEditors::RunTest(const FString& Parameters)
 {
diff --git a/Engine/Plugins/Developer/Concert/ConcertSync/ConcertSyncTest/Source/ConcertSyncTest/Private/ConcertDataStoreTests.cpp b/Engine/Plugins/Developer/Concert/ConcertSync/ConcertSyncTest/Source/ConcertSyncTest/Private/ConcertDataStoreTests.cpp
--- a/Engine/Plugins/Developer/Concert/ConcertSync/ConcertSyncTest/Source/ConcertSyncTest/Private/ConcertDataStoreTests.cpp
+++ b/Engine/Plugins/Developer/Concert/ConcertSync/ConcertSyncTest/Source/ConcertSyncTest/Private/ConcertDataStoreTests.cpp
@@ -513,7 +513,7 @@ private:
 } // namespace UE::ConcertSyncTests
 
 /** Ensures the Concert data store correctly versions the stored values. */
-IMPLEMENT_SIMPLE_AUTOMATION_TEST(FConcertDataStoreValueVersioning, "Editor.Concert.DataStore.ValueVersioning", ConcertDataStoreTestFlags)
+IMPLEMENT_SIMPLE_AUTOMATION_TEST(FConcertDataStoreValueVersioning, "VirtualProduction.Concert.DataStore.ValueVersioning", ConcertDataStoreTestFlags)
 
 bool FConcertDataStoreValueVersioning::RunTest(const FString& Parameters)
 {
@@ -599,7 +599,7 @@ bool FConcertDataStoreValueVersioning::RunTest(const FString& Parameters)
 }
 
 /** Ensures the Concert data store correctly handles the common operations. */
-IMPLEMENT_CUSTOM_SIMPLE_AUTOMATION_TEST(FConcertDataStoreClientServerCommonOperations, UE::ConcertSyncTests::DatastoreTests::FConcertDataStoreClientServerTest, "Editor.Concert.DataStore.ClientServer.CommonOperations", ConcertDataStoreTestFlags)
+IMPLEMENT_CUSTOM_SIMPLE_AUTOMATION_TEST(FConcertDataStoreClientServerCommonOperations, UE::ConcertSyncTests::DatastoreTests::FConcertDataStoreClientServerTest, "VirtualProduction.Concert.DataStore.ClientServer.CommonOperations", ConcertDataStoreTestFlags)
 
 bool FConcertDataStoreClientServerCommonOperations::RunTest(const FString& Parameters)
 {
@@ -627,7 +627,7 @@ bool FConcertDataStoreClientServerCommonOperations::RunTest(const FString& Param
 }
 
 /** Ensures the Concert data store correctly handles the "key not found" cases. */
-IMPLEMENT_CUSTOM_SIMPLE_AUTOMATION_TEST(FConcertDataStoreClientServerKeyNotFound, UE::ConcertSyncTests::DatastoreTests::FConcertDataStoreClientServerTest, "Editor.Concert.DataStore.ClientServer.KeyNotFound", ConcertDataStoreTestFlags)
+IMPLEMENT_CUSTOM_SIMPLE_AUTOMATION_TEST(FConcertDataStoreClientServerKeyNotFound, UE::ConcertSyncTests::DatastoreTests::FConcertDataStoreClientServerTest, "VirtualProduction.Concert.DataStore.ClientServer.KeyNotFound", ConcertDataStoreTestFlags)
 
 bool FConcertDataStoreClientServerKeyNotFound::RunTest(const FString& Parameters)
 {
@@ -645,7 +645,7 @@ bool FConcertDataStoreClientServerKeyNotFound::RunTest(const FString& Parameters
 }
 
 /** Ensures the Concert data store correctly handles the "type mismatch" cases. */
-IMPLEMENT_CUSTOM_SIMPLE_AUTOMATION_TEST(FConcertDataStoreClientServerTypeMismatch, UE::ConcertSyncTests::DatastoreTests::FConcertDataStoreClientServerTest, "Editor.Concert.DataStore.ClientServer.TypeMismatch", ConcertDataStoreTestFlags)
+IMPLEMENT_CUSTOM_SIMPLE_AUTOMATION_TEST(FConcertDataStoreClientServerTypeMismatch, UE::ConcertSyncTests::DatastoreTests::FConcertDataStoreClientServerTest, "VirtualProduction.Concert.DataStore.ClientServer.TypeMismatch", ConcertDataStoreTestFlags)
 
 bool FConcertDataStoreClientServerTypeMismatch::RunTest(const FString& Parameters)
 {
@@ -664,7 +664,7 @@ bool FConcertDataStoreClientServerTypeMismatch::RunTest(const FString& Parameter
 }
 
 /** Ensures the Concert data store optimize the "compare and exchange" operation, to avoid sending the payload when using the version is more optimal. */
-IMPLEMENT_CUSTOM_SIMPLE_AUTOMATION_TEST(FConcertDataStoreClientServerCompareExchangeOptimization, UE::ConcertSyncTests::DatastoreTests::FConcertDataStoreClientServerTest, "Editor.Concert.DataStore.ClientServer.CompareExchangeOptimization", ConcertDataStoreTestFlags)
+IMPLEMENT_CUSTOM_SIMPLE_AUTOMATION_TEST(FConcertDataStoreClientServerCompareExchangeOptimization, UE::ConcertSyncTests::DatastoreTests::FConcertDataStoreClientServerTest, "VirtualProduction.Concert.DataStore.ClientServer.CompareExchangeOptimization", ConcertDataStoreTestFlags)
 
 bool FConcertDataStoreClientServerCompareExchangeOptimization::RunTest(const FString& Parameters)
 {
@@ -693,7 +693,7 @@ bool FConcertDataStoreClientServerCompareExchangeOptimization::RunTest(const FSt
 }
 
 /** Ensures the Concert data store client correctly caches the key/values when it receives the response to its requests form the server. */
-IMPLEMENT_CUSTOM_SIMPLE_AUTOMATION_TEST(FConcertDataStoreClientServerClientCache, UE::ConcertSyncTests::DatastoreTests::FConcertDataStoreClientServerTest, "Editor.Concert.DataStore.ClientServer.ClientCache", ConcertDataStoreTestFlags)
+IMPLEMENT_CUSTOM_SIMPLE_AUTOMATION_TEST(FConcertDataStoreClientServerClientCache, UE::ConcertSyncTests::DatastoreTests::FConcertDataStoreClientServerTest, "VirtualProduction.Concert.DataStore.ClientServer.ClientCache", ConcertDataStoreTestFlags)
 
 bool FConcertDataStoreClientServerClientCache::RunTest(const FString& Parameters)
 {
@@ -796,7 +796,7 @@ bool FConcertDataStoreClientServerClientCache::RunTest(const FString& Parameters
 }
 
 /** Ensures the Concert data storeserver correctly push notifications to client and client populate its cache. */
-IMPLEMENT_CUSTOM_SIMPLE_AUTOMATION_TEST(FConcertDataStoreClientServerChangeNotification, UE::ConcertSyncTests::DatastoreTests::FConcertDataStoreClientServerTest, "Editor.Concert.DataStore.ClientServer.ChangeNotification", ConcertDataStoreTestFlags)
+IMPLEMENT_CUSTOM_SIMPLE_AUTOMATION_TEST(FConcertDataStoreClientServerChangeNotification, UE::ConcertSyncTests::DatastoreTests::FConcertDataStoreClientServerTest, "VirtualProduction.Concert.DataStore.ClientServer.ChangeNotification", ConcertDataStoreTestFlags)
 
 bool FConcertDataStoreClientServerChangeNotification::RunTest(const FString& Parameters)
 {
@@ -849,7 +849,7 @@ bool FConcertDataStoreClientServerChangeNotification::RunTest(const FString& Par
 }
 
 /** Ensures the Concert data store client correctly call the change handler. */
-IMPLEMENT_CUSTOM_SIMPLE_AUTOMATION_TEST(FConcertDataStoreClientServerChangeNotificationHandler, UE::ConcertSyncTests::DatastoreTests::FConcertDataStoreClientServerTest, "Editor.Concert.DataStore.ClientServer.ChangeNotificationHandler", ConcertDataStoreTestFlags)
+IMPLEMENT_CUSTOM_SIMPLE_AUTOMATION_TEST(FConcertDataStoreClientServerChangeNotificationHandler, UE::ConcertSyncTests::DatastoreTests::FConcertDataStoreClientServerTest, "VirtualProduction.Concert.DataStore.ClientServer.ChangeNotificationHandler", ConcertDataStoreTestFlags)
 
 bool FConcertDataStoreClientServerChangeNotificationHandler::RunTest(const FString& Parameters)
 {
@@ -968,7 +968,7 @@ bool FConcertDataStoreClientServerChangeNotificationHandler::RunTest(const FStri
 }
 
 /** Ensures the Concert data store blocking API works. */
-IMPLEMENT_CUSTOM_SIMPLE_AUTOMATION_TEST(FConcertDataStoreClientServerBlockingApi, UE::ConcertSyncTests::DatastoreTests::FConcertDataStoreClientServerTest, "Editor.Concert.DataStore.ClientServer.BlockingAPI", ConcertDataStoreTestFlags)
+IMPLEMENT_CUSTOM_SIMPLE_AUTOMATION_TEST(FConcertDataStoreClientServerBlockingApi, UE::ConcertSyncTests::DatastoreTests::FConcertDataStoreClientServerTest, "VirtualProduction.Concert.DataStore.ClientServer.BlockingAPI", ConcertDataStoreTestFlags)
 
 bool FConcertDataStoreClientServerBlockingApi::RunTest(const FString& Parameters)
 {
@@ -1024,7 +1024,7 @@ bool FConcertDataStoreClientServerBlockingApi::RunTest(const FString& Parameters
 }
 
 /** Ensures the Concert data store using continuation API works. */
-IMPLEMENT_CUSTOM_SIMPLE_AUTOMATION_TEST(FConcertDataStoreClientServerContinuationApi, UE::ConcertSyncTests::DatastoreTests::FConcertDataStoreClientServerTest, "Editor.Concert.DataStore.ClientServer.ContinuationAPI", ConcertDataStoreTestFlags)
+IMPLEMENT_CUSTOM_SIMPLE_AUTOMATION_TEST(FConcertDataStoreClientServerContinuationApi, UE::ConcertSyncTests::DatastoreTests::FConcertDataStoreClientServerTest, "VirtualProduction.Concert.DataStore.ClientServer.ContinuationAPI", ConcertDataStoreTestFlags)
 
 bool FConcertDataStoreClientServerContinuationApi::RunTest(const FString& Parameters)
 {
@@ -1082,7 +1082,7 @@ bool FConcertDataStoreClientServerContinuationApi::RunTest(const FString& Parame
 }
 
 /** Ensures the Concert data store local client store works. */
-IMPLEMENT_SIMPLE_AUTOMATION_TEST(FConcertDataStoreClientPrivateStore, "Editor.Concert.DataStore.ClientPrivateStore", ConcertDataStoreTestFlags)
+IMPLEMENT_SIMPLE_AUTOMATION_TEST(FConcertDataStoreClientPrivateStore, "VirtualProduction.Concert.DataStore.ClientPrivateStore", ConcertDataStoreTestFlags)
 
 bool FConcertDataStoreClientPrivateStore::RunTest(const FString& Parameters)
 {
diff --git a/Engine/Plugins/Developer/Concert/ConcertSync/ConcertSyncTest/Source/ConcertSyncTest/Private/Database/ConcertDatabaseTests.cpp b/Engine/Plugins/Developer/Concert/ConcertSync/ConcertSyncTest/Source/ConcertSyncTest/Private/Database/ConcertDatabaseTests.cpp
--- a/Engine/Plugins/Developer/Concert/ConcertSync/ConcertSyncTest/Source/ConcertSyncTest/Private/Database/ConcertDatabaseTests.cpp
+++ b/Engine/Plugins/Developer/Concert/ConcertSync/ConcertSyncTest/Source/ConcertSyncTest/Private/Database/ConcertDatabaseTests.cpp
@@ -12,7 +12,7 @@
 #include "Misc/Paths.h"
 #include "HAL/FileManager.h"
 
-IMPLEMENT_SIMPLE_AUTOMATION_TEST(FConcertSessionDatabaseTest, "Editor.Concert.Database.SessionDatabase", EAutomationTestFlags::EditorContext | EAutomationTestFlags::EngineFilter)
+IMPLEMENT_SIMPLE_AUTOMATION_TEST(FConcertSessionDatabaseTest, "VirtualProduction.Concert.Database.SessionDatabase", EAutomationTestFlags::EditorContext | EAutomationTestFlags::EngineFilter)
 
 bool FConcertSessionDatabaseTest::RunTest(const FString& Parameters)
 {
diff --git a/Engine/Plugins/Developer/Concert/ConcertSync/ConcertSyncTest/Source/ConcertSyncTest/Private/Database/ReplicationActivity.spec.cpp b/Engine/Plugins/Developer/Concert/ConcertSync/ConcertSyncTest/Source/ConcertSyncTest/Private/Database/ReplicationActivity.spec.cpp
--- a/Engine/Plugins/Developer/Concert/ConcertSync/ConcertSyncTest/Source/ConcertSyncTest/Private/Database/ReplicationActivity.spec.cpp
+++ b/Engine/Plugins/Developer/Concert/ConcertSync/ConcertSyncTest/Source/ConcertSyncTest/Private/Database/ReplicationActivity.spec.cpp
@@ -1,4 +1,4 @@
-﻿// Copyright Epic Games, Inc. All Rights Reserved.
+// Copyright Epic Games, Inc. All Rights Reserved.
 
 #include "ConcertSyncSessionDatabase.h"
 #include "Replication/Data/ConcertPropertySelection.h"
@@ -14,7 +14,7 @@
 
 namespace UE::ConcertSyncTests
 {
-	BEGIN_DEFINE_SPEC(FConcertDatabaseTest_ReplicationActivity, "Editor.Concert.Database.ReplicationActivity", EAutomationTestFlags::EditorContext | EAutomationTestFlags::EngineFilter)
+	BEGIN_DEFINE_SPEC(FConcertDatabaseTest_ReplicationActivity, "VirtualProduction.Concert.Database.ReplicationActivity", EAutomationTestFlags::EditorContext | EAutomationTestFlags::EngineFilter)
 		TUniquePtr<FConcertSyncSessionDatabase> Database;
 	
 		static FString GetDatabasePath() { return FPaths::ProjectIntermediateDir() / TEXT("ConcertDatabaseTest_Server"); }
diff --git a/Engine/Plugins/Developer/Concert/ConcertSync/ConcertSyncTest/Source/ConcertSyncTest/Private/HistoryEdition/DependencyGraphBuilderTests.cpp b/Engine/Plugins/Developer/Concert/ConcertSync/ConcertSyncTest/Source/ConcertSyncTest/Private/HistoryEdition/DependencyGraphBuilderTests.cpp
--- a/Engine/Plugins/Developer/Concert/ConcertSync/ConcertSyncTest/Source/ConcertSyncTest/Private/HistoryEdition/DependencyGraphBuilderTests.cpp
+++ b/Engine/Plugins/Developer/Concert/ConcertSync/ConcertSyncTest/Source/ConcertSyncTest/Private/HistoryEdition/DependencyGraphBuilderTests.cpp
@@ -40,7 +40,7 @@ namespace UE::ConcertSyncTests::RenameEditAndDeleteMapsFlowTest
 	 *	7 -> 5 (PackageRename)
 	 *	8 -> 5 (PackageRemoval)
 	 */
-	IMPLEMENT_SIMPLE_AUTOMATION_TEST(FRenameEditAndDeleteMapsFlowTest, "Editor.Concert.History.BuildGraph.RenameEditAndDeleteMapsFlow", EAutomationTestFlags::EditorContext | EAutomationTestFlags::EngineFilter)
+	IMPLEMENT_SIMPLE_AUTOMATION_TEST(FRenameEditAndDeleteMapsFlowTest, "VirtualProduction.Concert.History.BuildGraph.RenameEditAndDeleteMapsFlow", EAutomationTestFlags::EditorContext | EAutomationTestFlags::EngineFilter)
 	bool FRenameEditAndDeleteMapsFlowTest::RunTest(const FString& Parameters)
 	{
 		FScopedSessionDatabaseWithEndpoint SessionDatabase(*this);
@@ -189,7 +189,7 @@ namespace UE::ConcertSyncTests::DeletingAndRecreatingActorIsHardDependency
 	 *
 	 * 3 > 2 is a hard dependency (removing 2 would result in attempting to create the actor twice, which is invalid).
 	 */
-	IMPLEMENT_SIMPLE_AUTOMATION_TEST(FDeletingAndRecreatingActorIsHardDependency, "Editor.Concert.History.BuildGraph.DeletingAndRecreatingActorIsHardDependency", EAutomationTestFlags::EditorContext | EAutomationTestFlags::EngineFilter)
+	IMPLEMENT_SIMPLE_AUTOMATION_TEST(FDeletingAndRecreatingActorIsHardDependency, "VirtualProduction.Concert.History.BuildGraph.DeletingAndRecreatingActorIsHardDependency", EAutomationTestFlags::EditorContext | EAutomationTestFlags::EngineFilter)
 	bool FDeletingAndRecreatingActorIsHardDependency::RunTest(const FString& Parameters)
 	{
 		using namespace ConcertSyncCore;
@@ -306,7 +306,7 @@ namespace UE::ConcertSyncTests::IndirectPackageDependencyTest
 	 *
 	 * This is relevant because the actor's construction script may depend query data from the data asset.
 	 */
-	IMPLEMENT_SIMPLE_AUTOMATION_TEST(FIndirectPackageDependencyTest, "Editor.Concert.History.BuildGraph.IndirectPackageDependencyTest", EAutomationTestFlags::EditorContext | EAutomationTestFlags::EngineFilter)
+	IMPLEMENT_SIMPLE_AUTOMATION_TEST(FIndirectPackageDependencyTest, "VirtualProduction.Concert.History.BuildGraph.IndirectPackageDependencyTest", EAutomationTestFlags::EditorContext | EAutomationTestFlags::EngineFilter)
 	bool FIndirectPackageDependencyTest::RunTest(const FString& Parameters)
 	{
 		// TODO:
@@ -371,7 +371,7 @@ namespace UE::ConcertSyncTests::PackageEditedDependencyTest
 	 *  8 -> 7
 	 *  13 -> 12
 	 */
-	IMPLEMENT_SIMPLE_AUTOMATION_TEST(FPackageEditedDependencyTest, "Editor.Concert.History.BuildGraph.PackageEditedDependencyTest", EAutomationTestFlags::EditorContext | EAutomationTestFlags::EngineFilter)
+	IMPLEMENT_SIMPLE_AUTOMATION_TEST(FPackageEditedDependencyTest, "VirtualProduction.Concert.History.BuildGraph.PackageEditedDependencyTest", EAutomationTestFlags::EditorContext | EAutomationTestFlags::EngineFilter)
 	bool FPackageEditedDependencyTest::RunTest(const FString& Parameters)
 	{
 		using namespace ConcertSyncCore;
@@ -526,7 +526,7 @@ namespace UE::ConcertSyncTests::RenamingDependencyEdgeCases
 	 *  8 Add Bar
 	 *  9 Rename Bar to Foo
 	 */
-	IMPLEMENT_SIMPLE_AUTOMATION_TEST(FRenamingDependencyEdgeCaseTests, "Editor.Concert.History.BuildGraph.RenamingDependencyEdgeCaseTests", EAutomationTestFlags::EditorContext | EAutomationTestFlags::EngineFilter)
+	IMPLEMENT_SIMPLE_AUTOMATION_TEST(FRenamingDependencyEdgeCaseTests, "VirtualProduction.Concert.History.BuildGraph.RenamingDependencyEdgeCaseTests", EAutomationTestFlags::EditorContext | EAutomationTestFlags::EngineFilter)
 	bool FRenamingDependencyEdgeCaseTests::RunTest(const FString& Parameters)
 	{
 		using namespace ConcertSyncCore;
@@ -597,7 +597,7 @@ namespace UE::ConcertSyncTests::TransactionDependencyEdgeCases
 	
 	TArray<FActivityID> FillDatabase(FScopedSessionDatabaseWithEndpoint& SessionDatabase);
 	
-	IMPLEMENT_SIMPLE_AUTOMATION_TEST(FTransactionDependencyEdgeCaseTests, "Editor.Concert.History.BuildGraph.TransactionDependencyEdgeCaseTests", EAutomationTestFlags::EditorContext | EAutomationTestFlags::EngineFilter)
+	IMPLEMENT_SIMPLE_AUTOMATION_TEST(FTransactionDependencyEdgeCaseTests, "VirtualProduction.Concert.History.BuildGraph.TransactionDependencyEdgeCaseTests", EAutomationTestFlags::EditorContext | EAutomationTestFlags::EngineFilter)
 	bool FTransactionDependencyEdgeCaseTests::RunTest(const FString& Parameters)
 	{
 		using namespace ConcertSyncCore;
diff --git a/Engine/Plugins/Developer/Concert/ConcertSync/ConcertSyncTest/Source/ConcertSyncTest/Private/HistoryEdition/DependencyGraphTests.cpp b/Engine/Plugins/Developer/Concert/ConcertSync/ConcertSyncTest/Source/ConcertSyncTest/Private/HistoryEdition/DependencyGraphTests.cpp
--- a/Engine/Plugins/Developer/Concert/ConcertSync/ConcertSyncTest/Source/ConcertSyncTest/Private/HistoryEdition/DependencyGraphTests.cpp
+++ b/Engine/Plugins/Developer/Concert/ConcertSync/ConcertSyncTest/Source/ConcertSyncTest/Private/HistoryEdition/DependencyGraphTests.cpp
@@ -10,7 +10,7 @@
 namespace UE::ConcertSyncTests
 {
 	/** Ensures that the dependency graph API functions as intended, i.e. you add, connect and query nodes. */
-	IMPLEMENT_SIMPLE_AUTOMATION_TEST(FConcertDependencyGraphApiTest, "Editor.Concert.History.GraphApi.DependencyGraphApi", EAutomationTestFlags::EditorContext | EAutomationTestFlags::SmokeFilter)
+	IMPLEMENT_SIMPLE_AUTOMATION_TEST(FConcertDependencyGraphApiTest, "VirtualProduction.Concert.History.GraphApi.DependencyGraphApi", EAutomationTestFlags::EditorContext | EAutomationTestFlags::SmokeFilter)
 	bool FConcertDependencyGraphApiTest::RunTest(const FString& Parameters)
 	{
 		using namespace UE::ConcertSyncCore;
@@ -54,7 +54,7 @@ namespace UE::ConcertSyncTests
 	}
 
 	/** Ensures that the dependency graph API functions as intended, i.e. you add, connect and query nodes. */
-	IMPLEMENT_SIMPLE_AUTOMATION_TEST(FDisallowCyclesInGraph, "Editor.Concert.History.GraphApi.DisallowCyclesInGraph", EAutomationTestFlags::EditorContext | EAutomationTestFlags::SmokeFilter)
+	IMPLEMENT_SIMPLE_AUTOMATION_TEST(FDisallowCyclesInGraph, "VirtualProduction.Concert.History.GraphApi.DisallowCyclesInGraph", EAutomationTestFlags::EditorContext | EAutomationTestFlags::SmokeFilter)
 	bool FDisallowCyclesInGraph::RunTest(const FString& Parameters)
 	{
 		using namespace UE::ConcertSyncCore;
diff --git a/Engine/Plugins/Developer/Concert/ConcertSync/ConcertSyncTest/Source/ConcertSyncTest/Private/HistoryEdition/HistoryBottomUpAnalysisTests.cpp b/Engine/Plugins/Developer/Concert/ConcertSync/ConcertSyncTest/Source/ConcertSyncTest/Private/HistoryEdition/HistoryBottomUpAnalysisTests.cpp
--- a/Engine/Plugins/Developer/Concert/ConcertSync/ConcertSyncTest/Source/ConcertSyncTest/Private/HistoryEdition/HistoryBottomUpAnalysisTests.cpp
+++ b/Engine/Plugins/Developer/Concert/ConcertSync/ConcertSyncTest/Source/ConcertSyncTest/Private/HistoryEdition/HistoryBottomUpAnalysisTests.cpp
@@ -16,7 +16,7 @@ namespace UE::ConcertSyncTests::AnalysisTests
 	/**
 	 * Runs through a couple of use cases of unmuting an activity.
 	 */
-	IMPLEMENT_SIMPLE_AUTOMATION_TEST(FBottomUpUseCase, "Editor.Concert.History.Analysis.BottomUpUseCase", EAutomationTestFlags::EditorContext | EAutomationTestFlags::EngineFilter)
+	IMPLEMENT_SIMPLE_AUTOMATION_TEST(FBottomUpUseCase, "VirtualProduction.Concert.History.Analysis.BottomUpUseCase", EAutomationTestFlags::EditorContext | EAutomationTestFlags::EngineFilter)
 	bool FBottomUpUseCase::RunTest(const FString& Parameters)
 	{
 		using namespace RenameEditAndDeleteMapsFlowTest;
@@ -68,7 +68,7 @@ namespace UE::ConcertSyncTests::AnalysisTests
 	 * The test: unmute L.
 	 * We want R to be marked has a hard dependency.
 	 */
-	IMPLEMENT_SIMPLE_AUTOMATION_TEST(FBottomUpDiamond, "Editor.Concert.History.Analysis.BottomUpDiamond", EAutomationTestFlags::EditorContext | EAutomationTestFlags::EngineFilter)
+	IMPLEMENT_SIMPLE_AUTOMATION_TEST(FBottomUpDiamond, "VirtualProduction.Concert.History.Analysis.BottomUpDiamond", EAutomationTestFlags::EditorContext | EAutomationTestFlags::EngineFilter)
 	bool FBottomUpDiamond::RunTest(const FString& Parameters)
 	{
 		using namespace ConcertSyncCore;
diff --git a/Engine/Plugins/Developer/Concert/ConcertSync/ConcertSyncTest/Source/ConcertSyncTest/Private/HistoryEdition/HistoryTopDownAnalysisTests.cpp b/Engine/Plugins/Developer/Concert/ConcertSync/ConcertSyncTest/Source/ConcertSyncTest/Private/HistoryEdition/HistoryTopDownAnalysisTests.cpp
--- a/Engine/Plugins/Developer/Concert/ConcertSync/ConcertSyncTest/Source/ConcertSyncTest/Private/HistoryEdition/HistoryTopDownAnalysisTests.cpp
+++ b/Engine/Plugins/Developer/Concert/ConcertSync/ConcertSyncTest/Source/ConcertSyncTest/Private/HistoryEdition/HistoryTopDownAnalysisTests.cpp
@@ -57,7 +57,7 @@ namespace UE::ConcertSyncTests
 
 namespace UE::ConcertSyncTests::AnalysisTests
 {
-	IMPLEMENT_SIMPLE_AUTOMATION_TEST(FTopDownUseCase, "Editor.Concert.History.Analysis.TopDownUseCase", EAutomationTestFlags::EditorContext | EAutomationTestFlags::EngineFilter)
+	IMPLEMENT_SIMPLE_AUTOMATION_TEST(FTopDownUseCase, "VirtualProduction.Concert.History.Analysis.TopDownUseCase", EAutomationTestFlags::EditorContext | EAutomationTestFlags::EngineFilter)
 	bool FTopDownUseCase::RunTest(const FString& Parameters)
 	{
 		using namespace RenameEditAndDeleteMapsFlowTest;
@@ -131,7 +131,7 @@ namespace UE::ConcertSyncTests::AnalysisTests
 	 * The test: delete R.
 	 * We want L to be marked has a hard dependency.
 	 */
-	IMPLEMENT_SIMPLE_AUTOMATION_TEST(FTopDownDiamond, "Editor.Concert.History.Analysis.TopDownDiamond", EAutomationTestFlags::EditorContext | EAutomationTestFlags::EngineFilter)
+	IMPLEMENT_SIMPLE_AUTOMATION_TEST(FTopDownDiamond, "VirtualProduction.Concert.History.Analysis.TopDownDiamond", EAutomationTestFlags::EditorContext | EAutomationTestFlags::EngineFilter)
 	bool FTopDownDiamond::RunTest(const FString& Parameters)
 	{
 		using namespace ConcertSyncCore;
diff --git a/Engine/Plugins/Developer/Concert/ConcertSync/ConcertSyncTest/Source/ConcertSyncTest/Private/Misc/ObjectPathHierarchy.spec.cpp b/Engine/Plugins/Developer/Concert/ConcertSync/ConcertSyncTest/Source/ConcertSyncTest/Private/Misc/ObjectPathHierarchy.spec.cpp
--- a/Engine/Plugins/Developer/Concert/ConcertSync/ConcertSyncTest/Source/ConcertSyncTest/Private/Misc/ObjectPathHierarchy.spec.cpp
+++ b/Engine/Plugins/Developer/Concert/ConcertSync/ConcertSyncTest/Source/ConcertSyncTest/Private/Misc/ObjectPathHierarchy.spec.cpp
@@ -1,4 +1,4 @@
-﻿// Copyright Epic Games, Inc. All Rights Reserved.
+// Copyright Epic Games, Inc. All Rights Reserved.
 
 #include "Misc/AutomationTest.h"
 
@@ -11,7 +11,7 @@
 namespace UE::ConcertSyncTests
 {
 	/** Tests functions in ObjectPathUtils.h */
-	BEGIN_DEFINE_SPEC(FObjectPathHierarchySpec, "Editor.Concert.Components.ObjectPathHierarchy", EAutomationTestFlags::EditorContext | EAutomationTestFlags::EngineFilter)
+	BEGIN_DEFINE_SPEC(FObjectPathHierarchySpec, "VirtualProduction.Concert.Components.ObjectPathHierarchy", EAutomationTestFlags::EditorContext | EAutomationTestFlags::EngineFilter)
 		ConcertSyncCore::FObjectPathHierarchy Hierarchy;
 	END_DEFINE_SPEC(FObjectPathHierarchySpec);
 
diff --git a/Engine/Plugins/Developer/Concert/ConcertSync/ConcertSyncTest/Source/ConcertSyncTest/Private/Misc/ObjectPathUtils.spec.cpp b/Engine/Plugins/Developer/Concert/ConcertSync/ConcertSyncTest/Source/ConcertSyncTest/Private/Misc/ObjectPathUtils.spec.cpp
--- a/Engine/Plugins/Developer/Concert/ConcertSync/ConcertSyncTest/Source/ConcertSyncTest/Private/Misc/ObjectPathUtils.spec.cpp
+++ b/Engine/Plugins/Developer/Concert/ConcertSync/ConcertSyncTest/Source/ConcertSyncTest/Private/Misc/ObjectPathUtils.spec.cpp
@@ -1,4 +1,4 @@
-﻿// Copyright Epic Games, Inc. All Rights Reserved.
+// Copyright Epic Games, Inc. All Rights Reserved.
 
 #include "Misc/AutomationTest.h"
 
@@ -12,7 +12,7 @@
 namespace UE::ConcertSyncTests
 {
 	/** Tests functions in ObjectPathUtils.h */
-	BEGIN_DEFINE_SPEC(FObjectPathUtilsSpec, "Editor.Concert.Components.ObjectPathUtils", EAutomationTestFlags::EditorContext | EAutomationTestFlags::EngineFilter)
+	BEGIN_DEFINE_SPEC(FObjectPathUtilsSpec, "VirtualProduction.Concert.Components.ObjectPathUtils", EAutomationTestFlags::EditorContext | EAutomationTestFlags::EngineFilter)
 	END_DEFINE_SPEC(FObjectPathUtilsSpec);
 
 	void FObjectPathUtilsSpec::Define()
diff --git a/Engine/Plugins/Developer/Concert/ConcertSync/ConcertSyncTest/Source/ConcertSyncTest/Private/Replication/Bugs/ChangeObjectNameDuringReplication.cpp b/Engine/Plugins/Developer/Concert/ConcertSync/ConcertSyncTest/Source/ConcertSyncTest/Private/Replication/Bugs/ChangeObjectNameDuringReplication.cpp
--- a/Engine/Plugins/Developer/Concert/ConcertSync/ConcertSyncTest/Source/ConcertSyncTest/Private/Replication/Bugs/ChangeObjectNameDuringReplication.cpp
+++ b/Engine/Plugins/Developer/Concert/ConcertSync/ConcertSyncTest/Source/ConcertSyncTest/Private/Replication/Bugs/ChangeObjectNameDuringReplication.cpp
@@ -17,7 +17,7 @@ namespace UE::ConcertSyncTests::Replication::Bugs
 	 * This happens for example when replicating a construction script created component on a Blueprint actor and the actor runs a construction script,
 	 * which reinstances the components.
 	 */
-	IMPLEMENT_CUSTOM_SIMPLE_AUTOMATION_TEST(FChangeObjectNameDuringReplication, FSendReceiveWorldTestBase, "Editor.Concert.Replication.Bugs.ChangeObjectNameDuringReplication", EAutomationTestFlags::EditorContext | EAutomationTestFlags::EngineFilter);
+	IMPLEMENT_CUSTOM_SIMPLE_AUTOMATION_TEST(FChangeObjectNameDuringReplication, FSendReceiveWorldTestBase, "VirtualProduction.Concert.Replication.Bugs.ChangeObjectNameDuringReplication", EAutomationTestFlags::EditorContext | EAutomationTestFlags::EngineFilter);
 	bool FChangeObjectNameDuringReplication::RunTest(const FString& Parameters)
 	{
 		// The real replication bridge must be used because we're testing that the real bridge detects renaming
diff --git a/Engine/Plugins/Developer/Concert/ConcertSync/ConcertSyncTest/Source/ConcertSyncTest/Private/Replication/Bugs/RevertChange.spec.cpp b/Engine/Plugins/Developer/Concert/ConcertSync/ConcertSyncTest/Source/ConcertSyncTest/Private/Replication/Bugs/RevertChange.spec.cpp
--- a/Engine/Plugins/Developer/Concert/ConcertSync/ConcertSyncTest/Source/ConcertSyncTest/Private/Replication/Bugs/RevertChange.spec.cpp
+++ b/Engine/Plugins/Developer/Concert/ConcertSync/ConcertSyncTest/Source/ConcertSyncTest/Private/Replication/Bugs/RevertChange.spec.cpp
@@ -8,7 +8,7 @@
 
 namespace UE::ConcertSyncTests::Replication::UI
 {
-	BEGIN_DEFINE_SPEC(FRevertChange, "Editor.Concert.Replication.Bugs.RevertRequest", EAutomationTestFlags::EditorContext | EAutomationTestFlags::EngineFilter)
+	BEGIN_DEFINE_SPEC(FRevertChange, "VirtualProduction.Concert.Replication.Bugs.RevertRequest", EAutomationTestFlags::EditorContext | EAutomationTestFlags::EngineFilter)
 		TUniquePtr<FObjectTestReplicator> Replicator;
 		TUniquePtr<FReplicationServer> Server;
 		FReplicationClient* Client1;
diff --git a/Engine/Plugins/Developer/Concert/ConcertSync/ConcertSyncTest/Source/ConcertSyncTest/Private/Replication/ClientManager/ChangeStreamAndAuthorityEvents.spec.cpp b/Engine/Plugins/Developer/Concert/ConcertSync/ConcertSyncTest/Source/ConcertSyncTest/Private/Replication/ClientManager/ChangeStreamAndAuthorityEvents.spec.cpp
--- a/Engine/Plugins/Developer/Concert/ConcertSync/ConcertSyncTest/Source/ConcertSyncTest/Private/Replication/ClientManager/ChangeStreamAndAuthorityEvents.spec.cpp
+++ b/Engine/Plugins/Developer/Concert/ConcertSync/ConcertSyncTest/Source/ConcertSyncTest/Private/Replication/ClientManager/ChangeStreamAndAuthorityEvents.spec.cpp
@@ -9,7 +9,7 @@
 
 namespace UE::ConcertSyncTests::Replication::UI
 {
-	BEGIN_DEFINE_SPEC(FChangeStreamAndAuthorityEvents, "Editor.Concert.Replication.ClientManager.ChangeStreamAndAuthorityEvents", EAutomationTestFlags::EditorContext | EAutomationTestFlags::EngineFilter)
+	BEGIN_DEFINE_SPEC(FChangeStreamAndAuthorityEvents, "VirtualProduction.Concert.Replication.ClientManager.ChangeStreamAndAuthorityEvents", EAutomationTestFlags::EditorContext | EAutomationTestFlags::EngineFilter)
 		TUniquePtr<FObjectTestReplicator> Replicator;
 		TUniquePtr<FReplicationServer> Server;
 		FReplicationClient* Sender;
diff --git a/Engine/Plugins/Developer/Concert/ConcertSync/ConcertSyncTest/Source/ConcertSyncTest/Private/Replication/ClientManager/PutClientStateEvents.spec.cpp b/Engine/Plugins/Developer/Concert/ConcertSync/ConcertSyncTest/Source/ConcertSyncTest/Private/Replication/ClientManager/PutClientStateEvents.spec.cpp
--- a/Engine/Plugins/Developer/Concert/ConcertSync/ConcertSyncTest/Source/ConcertSyncTest/Private/Replication/ClientManager/PutClientStateEvents.spec.cpp
+++ b/Engine/Plugins/Developer/Concert/ConcertSync/ConcertSyncTest/Source/ConcertSyncTest/Private/Replication/ClientManager/PutClientStateEvents.spec.cpp
@@ -9,7 +9,7 @@
 
 namespace UE::ConcertSyncTests::Replication::UI
 {
-	BEGIN_DEFINE_SPEC(FPutClientStateEvents, "Editor.Concert.Replication.ClientManager.PutClientStateEvents", EAutomationTestFlags::EditorContext | EAutomationTestFlags::EngineFilter)
+	BEGIN_DEFINE_SPEC(FPutClientStateEvents, "VirtualProduction.Concert.Replication.ClientManager.PutClientStateEvents", EAutomationTestFlags::EditorContext | EAutomationTestFlags::EngineFilter)
 		TUniquePtr<FObjectTestReplicator> Replicator_Client1;
 		TUniquePtr<FObjectTestReplicator> Replicator_Client2;
 		TUniquePtr<FReplicationServer> Server;
diff --git a/Engine/Plugins/Developer/Concert/ConcertSync/ConcertSyncTest/Source/ConcertSyncTest/Private/Replication/ClientManager/RestoreContentEvents.spec.cpp b/Engine/Plugins/Developer/Concert/ConcertSync/ConcertSyncTest/Source/ConcertSyncTest/Private/Replication/ClientManager/RestoreContentEvents.spec.cpp
--- a/Engine/Plugins/Developer/Concert/ConcertSync/ConcertSyncTest/Source/ConcertSyncTest/Private/Replication/ClientManager/RestoreContentEvents.spec.cpp
+++ b/Engine/Plugins/Developer/Concert/ConcertSync/ConcertSyncTest/Source/ConcertSyncTest/Private/Replication/ClientManager/RestoreContentEvents.spec.cpp
@@ -27,7 +27,7 @@ namespace UE::ConcertSyncTests::Replication::UI
 		return Result;
 	}
 	
-	BEGIN_DEFINE_SPEC(FRestoreContentEvents, "Editor.Concert.Replication.ClientManager.RestoreContentEvents", EAutomationTestFlags::EditorContext | EAutomationTestFlags::EngineFilter)
+	BEGIN_DEFINE_SPEC(FRestoreContentEvents, "VirtualProduction.Concert.Replication.ClientManager.RestoreContentEvents", EAutomationTestFlags::EditorContext | EAutomationTestFlags::EngineFilter)
 		TUniquePtr<FObjectTestReplicator> Replicator;
 		TUniquePtr<FReplicationServer> Server;
 		FReplicationClient* Sender;
diff --git a/Engine/Plugins/Developer/Concert/ConcertSync/ConcertSyncTest/Source/ConcertSyncTest/Private/Replication/Components/BacktrackActivityHistory.spec.cpp b/Engine/Plugins/Developer/Concert/ConcertSync/ConcertSyncTest/Source/ConcertSyncTest/Private/Replication/Components/BacktrackActivityHistory.spec.cpp
--- a/Engine/Plugins/Developer/Concert/ConcertSync/ConcertSyncTest/Source/ConcertSyncTest/Private/Replication/Components/BacktrackActivityHistory.spec.cpp
+++ b/Engine/Plugins/Developer/Concert/ConcertSync/ConcertSyncTest/Source/ConcertSyncTest/Private/Replication/Components/BacktrackActivityHistory.spec.cpp
@@ -1,4 +1,4 @@
-﻿// Copyright Epic Games, Inc. All Rights Reserved.
+// Copyright Epic Games, Inc. All Rights Reserved.
 
 #include "ConcertSyncSessionTypes.h"
 #include "Replication/Messages/ReplicationActivity.h"
@@ -117,7 +117,7 @@ namespace UE::ConcertSyncTests::Replication::BacktrackActivityHistory
 		//~ End IConcertClientWorkspace Interface
 	};
 	
-	BEGIN_DEFINE_SPEC(FBacktrackActivityHistory, "Editor.Concert.Replication.Components.BacktrackActivityHistory", EAutomationTestFlags::EditorContext | EAutomationTestFlags::EngineFilter)
+	BEGIN_DEFINE_SPEC(FBacktrackActivityHistory, "VirtualProduction.Concert.Replication.Components.BacktrackActivityHistory", EAutomationTestFlags::EditorContext | EAutomationTestFlags::EngineFilter)
 		TUniquePtr<FTestCase> TestData;
 	END_DEFINE_SPEC(FBacktrackActivityHistory);
 
diff --git a/Engine/Plugins/Developer/Concert/ConcertSync/ConcertSyncTest/Source/ConcertSyncTest/Private/Replication/Components/CombineMuteRequests.spec.cpp b/Engine/Plugins/Developer/Concert/ConcertSync/ConcertSyncTest/Source/ConcertSyncTest/Private/Replication/Components/CombineMuteRequests.spec.cpp
--- a/Engine/Plugins/Developer/Concert/ConcertSync/ConcertSyncTest/Source/ConcertSyncTest/Private/Replication/Components/CombineMuteRequests.spec.cpp
+++ b/Engine/Plugins/Developer/Concert/ConcertSync/ConcertSyncTest/Source/ConcertSyncTest/Private/Replication/Components/CombineMuteRequests.spec.cpp
@@ -62,7 +62,7 @@ namespace UE::ConcertSyncTests::Replication
 	};
 	
 	/** This tests that muting & unmuting retains a consistent check across various types of changes. */
-	BEGIN_DEFINE_SPEC(FCombineMuteRequestsSpec, "Editor.Concert.Replication.Components.CombineMuteRequests", EAutomationTestFlags::EditorContext | EAutomationTestFlags::EngineFilter)
+	BEGIN_DEFINE_SPEC(FCombineMuteRequestsSpec, "VirtualProduction.Concert.Replication.Components.CombineMuteRequests", EAutomationTestFlags::EditorContext | EAutomationTestFlags::EngineFilter)
 		FGuid StreamId					= FGuid::NewGuid();
 		// Leverage FObjectTestReplicator to create more UObjects
 		TSharedPtr<FObjectTestReplicator> Replicator_Foo;			
diff --git a/Engine/Plugins/Developer/Concert/ConcertSync/ConcertSyncTest/Source/ConcertSyncTest/Private/Replication/Components/FrequencySettingsTests.cpp b/Engine/Plugins/Developer/Concert/ConcertSync/ConcertSyncTest/Source/ConcertSyncTest/Private/Replication/Components/FrequencySettingsTests.cpp
--- a/Engine/Plugins/Developer/Concert/ConcertSync/ConcertSyncTest/Source/ConcertSyncTest/Private/Replication/Components/FrequencySettingsTests.cpp
+++ b/Engine/Plugins/Developer/Concert/ConcertSync/ConcertSyncTest/Source/ConcertSyncTest/Private/Replication/Components/FrequencySettingsTests.cpp
@@ -10,7 +10,7 @@
 namespace UE::ConcertSyncTests
 {
 	/** Tests that FConcertObjectReplicationSettings's < and <= work as expected. */
-	IMPLEMENT_SIMPLE_AUTOMATION_TEST(FFrequencySettingsLessEqual, "Editor.Concert.Replication.Components.FrequencySettingsLessEqual", EAutomationTestFlags::EditorContext | EAutomationTestFlags::EngineFilter);
+	IMPLEMENT_SIMPLE_AUTOMATION_TEST(FFrequencySettingsLessEqual, "VirtualProduction.Concert.Replication.Components.FrequencySettingsLessEqual", EAutomationTestFlags::EditorContext | EAutomationTestFlags::EngineFilter);
 	bool FFrequencySettingsLessEqual::RunTest(const FString& Parameters)
 	{
 		constexpr FConcertObjectReplicationSettings FPS_30 { EConcertObjectReplicationMode::SpecifiedRate, 30 }; 
@@ -38,7 +38,7 @@ namespace UE::ConcertSyncTests
 	}
 
 	/** Tests that RelativeLocation and RelativeRotation properties are discovered automatically by FConcertStreamObjectAutoBindingRules. */
-	IMPLEMENT_SIMPLE_AUTOMATION_TEST(FAutoBindingPropertySettingsTest, "Editor.Concert.Replication.Components.AutoBindingPropertySettingsTest", EAutomationTestFlags::EditorContext | EAutomationTestFlags::EngineFilter);
+	IMPLEMENT_SIMPLE_AUTOMATION_TEST(FAutoBindingPropertySettingsTest, "VirtualProduction.Concert.Replication.Components.AutoBindingPropertySettingsTest", EAutomationTestFlags::EditorContext | EAutomationTestFlags::EngineFilter);
 	bool FAutoBindingPropertySettingsTest::RunTest(const FString& Parameters)
 	{
 		// 1. Set up: rules
diff --git a/Engine/Plugins/Developer/Concert/ConcertSync/ConcertSyncTest/Source/ConcertSyncTest/Private/Replication/Components/FrequencyTrackerTests.cpp b/Engine/Plugins/Developer/Concert/ConcertSync/ConcertSyncTest/Source/ConcertSyncTest/Private/Replication/Components/FrequencyTrackerTests.cpp
--- a/Engine/Plugins/Developer/Concert/ConcertSync/ConcertSyncTest/Source/ConcertSyncTest/Private/Replication/Components/FrequencyTrackerTests.cpp
+++ b/Engine/Plugins/Developer/Concert/ConcertSync/ConcertSyncTest/Source/ConcertSyncTest/Private/Replication/Components/FrequencyTrackerTests.cpp
@@ -15,7 +15,7 @@ namespace UE::ConcertSyncTests
 	 *	- Tracks objects correctly,
 	 *	- Does not leak objects when cleaning cache
 	 */
-	IMPLEMENT_SIMPLE_AUTOMATION_TEST(FFrequencyTrackerTest, "Editor.Concert.Replication.Components.FrequencyTracker", EAutomationTestFlags::EditorContext | EAutomationTestFlags::EngineFilter);
+	IMPLEMENT_SIMPLE_AUTOMATION_TEST(FFrequencyTrackerTest, "VirtualProduction.Concert.Replication.Components.FrequencyTracker", EAutomationTestFlags::EditorContext | EAutomationTestFlags::EngineFilter);
 	bool FFrequencyTrackerTest::RunTest(const FString& Parameters)
 	{
 		using namespace ConcertSyncCore;
diff --git a/Engine/Plugins/Developer/Concert/ConcertSync/ConcertSyncTest/Source/ConcertSyncTest/Private/Replication/Components/LabelRemapping/DifferentWorld.spec.cpp b/Engine/Plugins/Developer/Concert/ConcertSync/ConcertSyncTest/Source/ConcertSyncTest/Private/Replication/Components/LabelRemapping/DifferentWorld.spec.cpp
--- a/Engine/Plugins/Developer/Concert/ConcertSync/ConcertSyncTest/Source/ConcertSyncTest/Private/Replication/Components/LabelRemapping/DifferentWorld.spec.cpp
+++ b/Engine/Plugins/Developer/Concert/ConcertSync/ConcertSyncTest/Source/ConcertSyncTest/Private/Replication/Components/LabelRemapping/DifferentWorld.spec.cpp
@@ -11,7 +11,7 @@
 namespace UE::ConcertSyncTests::Replication::ActorLabelRemapping
 {
 	/** This test that remapping works when the target world is different from the origin world. */
-	BEGIN_DEFINE_SPEC(FActorLabelRemappingDifferentWorldsSpec, "Editor.Concert.Replication.Components.ActorLabelRemapping.DifferentWorlds", EAutomationTestFlags::EditorContext | EAutomationTestFlags::EngineFilter)
+	BEGIN_DEFINE_SPEC(FActorLabelRemappingDifferentWorldsSpec, "VirtualProduction.Concert.Replication.Components.ActorLabelRemapping.DifferentWorlds", EAutomationTestFlags::EditorContext | EAutomationTestFlags::EngineFilter)
 		TUniquePtr<FReplicatedTestWorld> OriginWorld;
 		TUniquePtr<FReplicatedTestWorld> TargetWorld;
 
diff --git a/Engine/Plugins/Developer/Concert/ConcertSync/ConcertSyncTest/Source/ConcertSyncTest/Private/Replication/Components/LabelRemapping/GenerateRemappingData.spec.cpp b/Engine/Plugins/Developer/Concert/ConcertSync/ConcertSyncTest/Source/ConcertSyncTest/Private/Replication/Components/LabelRemapping/GenerateRemappingData.spec.cpp
--- a/Engine/Plugins/Developer/Concert/ConcertSync/ConcertSyncTest/Source/ConcertSyncTest/Private/Replication/Components/LabelRemapping/GenerateRemappingData.spec.cpp
+++ b/Engine/Plugins/Developer/Concert/ConcertSync/ConcertSyncTest/Source/ConcertSyncTest/Private/Replication/Components/LabelRemapping/GenerateRemappingData.spec.cpp
@@ -10,7 +10,7 @@
 namespace UE::ConcertSyncTests::Replication::ActorLabelRemapping
 {
 	/** This test that remapping works after changes have been made to the world. */
-	BEGIN_DEFINE_SPEC(FGenerateRemappingDataSpec, "Editor.Concert.Replication.Components.ActorLabelRemapping.GenerateRemappingDataSpec", EAutomationTestFlags::EditorContext | EAutomationTestFlags::EngineFilter)
+	BEGIN_DEFINE_SPEC(FGenerateRemappingDataSpec, "VirtualProduction.Concert.Replication.Components.ActorLabelRemapping.GenerateRemappingDataSpec", EAutomationTestFlags::EditorContext | EAutomationTestFlags::EngineFilter)
 		TUniquePtr<FReplicatedTestWorld> ReplicatedWorld;
 
 		UWorld& GetWorld() const { return *ReplicatedWorld->PreviewScene.GetWorld(); }
diff --git a/Engine/Plugins/Developer/Concert/ConcertSync/ConcertSyncTest/Source/ConcertSyncTest/Private/Replication/Components/LabelRemapping/SameWorld.spec.cpp b/Engine/Plugins/Developer/Concert/ConcertSync/ConcertSyncTest/Source/ConcertSyncTest/Private/Replication/Components/LabelRemapping/SameWorld.spec.cpp
--- a/Engine/Plugins/Developer/Concert/ConcertSync/ConcertSyncTest/Source/ConcertSyncTest/Private/Replication/Components/LabelRemapping/SameWorld.spec.cpp
+++ b/Engine/Plugins/Developer/Concert/ConcertSync/ConcertSyncTest/Source/ConcertSyncTest/Private/Replication/Components/LabelRemapping/SameWorld.spec.cpp
@@ -13,7 +13,7 @@
 namespace UE::ConcertSyncTests::Replication::ActorLabelRemapping
 {
 	/** This test that remapping works after changes have been made to the world. */
-	BEGIN_DEFINE_SPEC(FActorLabelRemappingSameWorldSpec, "Editor.Concert.Replication.Components.ActorLabelRemapping.SameWorld", EAutomationTestFlags::EditorContext | EAutomationTestFlags::EngineFilter)
+	BEGIN_DEFINE_SPEC(FActorLabelRemappingSameWorldSpec, "VirtualProduction.Concert.Replication.Components.ActorLabelRemapping.SameWorld", EAutomationTestFlags::EditorContext | EAutomationTestFlags::EngineFilter)
 		TUniquePtr<FReplicatedTestWorld> ReplicatedWorld;
 
 		UWorld& GetWorld() const { return *ReplicatedWorld->PreviewScene.GetWorld(); }
diff --git a/Engine/Plugins/Developer/Concert/ConcertSync/ConcertSyncTest/Source/ConcertSyncTest/Private/Replication/Components/ReplicatedObjectHierarchyCache.spec.cpp b/Engine/Plugins/Developer/Concert/ConcertSync/ConcertSyncTest/Source/ConcertSyncTest/Private/Replication/Components/ReplicatedObjectHierarchyCache.spec.cpp
--- a/Engine/Plugins/Developer/Concert/ConcertSync/ConcertSyncTest/Source/ConcertSyncTest/Private/Replication/Components/ReplicatedObjectHierarchyCache.spec.cpp
+++ b/Engine/Plugins/Developer/Concert/ConcertSync/ConcertSyncTest/Source/ConcertSyncTest/Private/Replication/Components/ReplicatedObjectHierarchyCache.spec.cpp
@@ -12,7 +12,7 @@
 namespace UE::ConcertSyncTests::Replication
 {
 	/** Tests functions in ObjectPathUtils.h */
-	BEGIN_DEFINE_SPEC(FReplicatedObjectHierarchyCacheSpec, "Editor.Concert.Replication.Components.ReplicatedObjectHierarchyCache", EAutomationTestFlags::EditorContext | EAutomationTestFlags::EngineFilter)
+	BEGIN_DEFINE_SPEC(FReplicatedObjectHierarchyCacheSpec, "VirtualProduction.Concert.Replication.Components.ReplicatedObjectHierarchyCache", EAutomationTestFlags::EditorContext | EAutomationTestFlags::EngineFilter)
 		FGuid ClientId_One = FGuid::NewGuid(); 
 		FGuid ClientId_Two = FGuid::NewGuid(); 
 		FGuid ClientId_Three = FGuid::NewGuid(); 
diff --git a/Engine/Plugins/Developer/Concert/ConcertSync/ConcertSyncTest/Source/ConcertSyncTest/Private/Replication/Components/ReplicationCache.spec.cpp b/Engine/Plugins/Developer/Concert/ConcertSync/ConcertSyncTest/Source/ConcertSyncTest/Private/Replication/Components/ReplicationCache.spec.cpp
--- a/Engine/Plugins/Developer/Concert/ConcertSync/ConcertSyncTest/Source/ConcertSyncTest/Private/Replication/Components/ReplicationCache.spec.cpp
+++ b/Engine/Plugins/Developer/Concert/ConcertSync/ConcertSyncTest/Source/ConcertSyncTest/Private/Replication/Components/ReplicationCache.spec.cpp
@@ -128,7 +128,7 @@ namespace UE::ConcertSyncTests
 		}
 	};
 	
-	BEGIN_DEFINE_SPEC(FReplicationCacheSpec, "Editor.Concert.Replication.Components.ReplicationCache", EAutomationTestFlags::EditorContext | EAutomationTestFlags::EngineFilter)
+	BEGIN_DEFINE_SPEC(FReplicationCacheSpec, "VirtualProduction.Concert.Replication.Components.ReplicationCache", EAutomationTestFlags::EditorContext | EAutomationTestFlags::EngineFilter)
 		// Set up data
 		const FGuid StreamId = FGuid::NewGuid();
 		const FSoftObjectPath ObjectPath { TEXT("/Game/World.World:PersistentLevel.StaticMeshActor0") };
diff --git a/Engine/Plugins/Developer/Concert/ConcertSync/ConcertSyncTest/Source/ConcertSyncTest/Private/Replication/Components/ReplicationChangeStreamUtilsTests.cpp b/Engine/Plugins/Developer/Concert/ConcertSync/ConcertSyncTest/Source/ConcertSyncTest/Private/Replication/Components/ReplicationChangeStreamUtilsTests.cpp
--- a/Engine/Plugins/Developer/Concert/ConcertSync/ConcertSyncTest/Source/ConcertSyncTest/Private/Replication/Components/ReplicationChangeStreamUtilsTests.cpp
+++ b/Engine/Plugins/Developer/Concert/ConcertSync/ConcertSyncTest/Source/ConcertSyncTest/Private/Replication/Components/ReplicationChangeStreamUtilsTests.cpp
@@ -19,7 +19,7 @@ namespace UE::ConcertSyncTests::ChangeStreamUtils
 	 *  - adding and removing objects
 	 *  - modify properties and class
 	 */
-	IMPLEMENT_SIMPLE_AUTOMATION_TEST(FSimpleDiffChangesTest, "Editor.Concert.Replication.Components.DiffChangesSimple", EAutomationTestFlags::EditorContext | EAutomationTestFlags::EngineFilter);
+	IMPLEMENT_SIMPLE_AUTOMATION_TEST(FSimpleDiffChangesTest, "VirtualProduction.Concert.Replication.Components.DiffChangesSimple", EAutomationTestFlags::EditorContext | EAutomationTestFlags::EngineFilter);
 	bool FSimpleDiffChangesTest::RunTest(const FString& Parameters)
 	{
 		// 1. Set up
diff --git a/Engine/Plugins/Developer/Concert/ConcertSync/ConcertSyncTest/Source/ConcertSyncTest/Private/Replication/Components/SyncControlState.spec.cpp b/Engine/Plugins/Developer/Concert/ConcertSync/ConcertSyncTest/Source/ConcertSyncTest/Private/Replication/Components/SyncControlState.spec.cpp
--- a/Engine/Plugins/Developer/Concert/ConcertSync/ConcertSyncTest/Source/ConcertSyncTest/Private/Replication/Components/SyncControlState.spec.cpp
+++ b/Engine/Plugins/Developer/Concert/ConcertSync/ConcertSyncTest/Source/ConcertSyncTest/Private/Replication/Components/SyncControlState.spec.cpp
@@ -8,7 +8,7 @@
 
 namespace UE::ConcertSyncTests::Replication::UI
 {
-	BEGIN_DEFINE_SPEC(FSyncControlStateSpec, "Editor.Concert.Replication.Components", EAutomationTestFlags::EditorContext | EAutomationTestFlags::EngineFilter)
+	BEGIN_DEFINE_SPEC(FSyncControlStateSpec, "VirtualProduction.Concert.Replication.Components", EAutomationTestFlags::EditorContext | EAutomationTestFlags::EngineFilter)
 	FGuid SenderStreamId = FGuid::NewGuid();
 	/** Leverage to create a real FSoftObjectPath. */
 	TSharedPtr<FObjectTestReplicator> ObjectReplicator;
diff --git a/Engine/Plugins/Developer/Concert/ConcertSync/ConcertSyncTest/Source/ConcertSyncTest/Private/Replication/ConcertPropertyTests.cpp b/Engine/Plugins/Developer/Concert/ConcertSync/ConcertSyncTest/Source/ConcertSyncTest/Private/Replication/ConcertPropertyTests.cpp
--- a/Engine/Plugins/Developer/Concert/ConcertSync/ConcertSyncTest/Source/ConcertSyncTest/Private/Replication/ConcertPropertyTests.cpp
+++ b/Engine/Plugins/Developer/Concert/ConcertSync/ConcertSyncTest/Source/ConcertSyncTest/Private/Replication/ConcertPropertyTests.cpp
@@ -21,7 +21,7 @@ namespace UE::ConcertSyncTest::Replication::PropertyChain
 	 *		- "normal" structs, like FVector > lists sub-properties
 	 *		- primitives and native structs > Path contains FConcertPropertyChain::InternalContainerPropertyValueName ("Value") at the end
 	 */
-	IMPLEMENT_SIMPLE_AUTOMATION_TEST(FConcertPropertyChainTests, "Editor.Concert.Replication.Data.PropertyChainTests", EAutomationTestFlags::EditorContext | EAutomationTestFlags::EngineFilter)
+	IMPLEMENT_SIMPLE_AUTOMATION_TEST(FConcertPropertyChainTests, "VirtualProduction.Concert.Replication.Data.PropertyChainTests", EAutomationTestFlags::EditorContext | EAutomationTestFlags::EngineFilter)
 	bool FConcertPropertyChainTests::RunTest(const FString& Parameters)
 	{
 		// The test is intentionally set up to make it it easy to set breakpoints albeit at the expense of increasing the amount of code
@@ -117,7 +117,7 @@ namespace UE::ConcertSyncTest::Replication::PropertyChain
 	/**
 	 * Tests that all property cases on the specially constructed example class UTestReflectionObject.
 	 */
-	IMPLEMENT_SIMPLE_AUTOMATION_TEST(FConcertForEachReplictableConcertPropertyTests, "Editor.Concert.Replication.Data.ForEachReplicatableConcertProperty", EAutomationTestFlags::EditorContext | EAutomationTestFlags::EngineFilter)
+	IMPLEMENT_SIMPLE_AUTOMATION_TEST(FConcertForEachReplictableConcertPropertyTests, "VirtualProduction.Concert.Replication.Data.ForEachReplicatableConcertProperty", EAutomationTestFlags::EditorContext | EAutomationTestFlags::EngineFilter)
 	bool FConcertForEachReplictableConcertPropertyTests::RunTest(const FString& Parameters)
 	{
 		UClass* Class = UTestReflectionObject::StaticClass();
@@ -217,7 +217,7 @@ namespace UE::ConcertSyncTest::Replication::PropertyChain
 	/**
 	 * Tests FConcertPropertyChain::Matches.
 	 */
-	IMPLEMENT_SIMPLE_AUTOMATION_TEST(FMatchConcertPropertyChainTest, "Editor.Concert.Replication.Data.MatchConcertPropertyChain", EAutomationTestFlags::EditorContext | EAutomationTestFlags::EngineFilter)
+	IMPLEMENT_SIMPLE_AUTOMATION_TEST(FMatchConcertPropertyChainTest, "VirtualProduction.Concert.Replication.Data.MatchConcertPropertyChain", EAutomationTestFlags::EditorContext | EAutomationTestFlags::EngineFilter)
 	bool FMatchConcertPropertyChainTest::RunTest(const FString& Parameters)
 	{
 		// 1. Set up
@@ -273,7 +273,7 @@ namespace UE::ConcertSyncTest::Replication::PropertyChain
 	 * Tests that UE::ConcertSyncCore::FReplicationPropertyFilter works correctly.
 	 * This uses the filter and serializes an object, doing a Memcmp to detect that exactly the expected properties are serialized.
 	 */
-	IMPLEMENT_SIMPLE_AUTOMATION_TEST(FSerializeWithPropertyFilterTest, "Editor.Concert.Replication.Data.SerializeWithPropertyFilter", EAutomationTestFlags::EditorContext | EAutomationTestFlags::EngineFilter)
+	IMPLEMENT_SIMPLE_AUTOMATION_TEST(FSerializeWithPropertyFilterTest, "VirtualProduction.Concert.Replication.Data.SerializeWithPropertyFilter", EAutomationTestFlags::EditorContext | EAutomationTestFlags::EngineFilter)
 	bool FSerializeWithPropertyFilterTest::RunTest(const FString& Parameters)
 	{
 		class FTestArchive : public FMemoryWriter
@@ -439,7 +439,7 @@ namespace UE::ConcertSyncTest::Replication::PropertyChain
 	 * Tests that UE::ConcertSyncCore::FReplicationPropertyFilter works correctly.
 	 * This uses the filter and serializes an object, doing a Memcmp to detect that exactly the expected properties are serialized.
 	 */
-	IMPLEMENT_SIMPLE_AUTOMATION_TEST(FOverlapPropertiesTest, "Editor.Concert.Replication.Data.OverlapProperties", EAutomationTestFlags::EditorContext | EAutomationTestFlags::EngineFilter)
+	IMPLEMENT_SIMPLE_AUTOMATION_TEST(FOverlapPropertiesTest, "VirtualProduction.Concert.Replication.Data.OverlapProperties", EAutomationTestFlags::EditorContext | EAutomationTestFlags::EngineFilter)
 	bool FOverlapPropertiesTest::RunTest(const FString& Parameters)
 	{
 		// 1. Define paths
diff --git a/Engine/Plugins/Developer/Concert/ConcertSync/ConcertSyncTest/Source/ConcertSyncTest/Private/Replication/DirtyPackageOnReplication.spec.cpp b/Engine/Plugins/Developer/Concert/ConcertSync/ConcertSyncTest/Source/ConcertSyncTest/Private/Replication/DirtyPackageOnReplication.spec.cpp
--- a/Engine/Plugins/Developer/Concert/ConcertSync/ConcertSyncTest/Source/ConcertSyncTest/Private/Replication/DirtyPackageOnReplication.spec.cpp
+++ b/Engine/Plugins/Developer/Concert/ConcertSync/ConcertSyncTest/Source/ConcertSyncTest/Private/Replication/DirtyPackageOnReplication.spec.cpp
@@ -1,4 +1,4 @@
-﻿// Copyright Epic Games, Inc. All Rights Reserved.
+// Copyright Epic Games, Inc. All Rights Reserved.
 
 #include "Misc/AutomationTest.h"
 #include "Util/ConcertClientReplicationBridgeMock.h"
@@ -8,7 +8,7 @@
 
 namespace UE::ConcertSyncTests::Replication::UI
 {
-	BEGIN_DEFINE_SPEC(FDirtyPackageOnReplication, "Editor.Concert.Replication.DirtyPackageOnReplication", EAutomationTestFlags::EditorContext | EAutomationTestFlags::EngineFilter)
+	BEGIN_DEFINE_SPEC(FDirtyPackageOnReplication, "VirtualProduction.Concert.Replication.DirtyPackageOnReplication", EAutomationTestFlags::EditorContext | EAutomationTestFlags::EngineFilter)
 		UPackage* Package;
 		TUniquePtr<FObjectTestReplicator> Replicator;
 		TUniquePtr<FReplicationServer> Server;
diff --git a/Engine/Plugins/Developer/Concert/ConcertSync/ConcertSyncTest/Source/ConcertSyncTest/Private/Replication/FrequencyTests.cpp b/Engine/Plugins/Developer/Concert/ConcertSync/ConcertSyncTest/Source/ConcertSyncTest/Private/Replication/FrequencyTests.cpp
--- a/Engine/Plugins/Developer/Concert/ConcertSync/ConcertSyncTest/Source/ConcertSyncTest/Private/Replication/FrequencyTests.cpp
+++ b/Engine/Plugins/Developer/Concert/ConcertSync/ConcertSyncTest/Source/ConcertSyncTest/Private/Replication/FrequencyTests.cpp
@@ -54,7 +54,7 @@ namespace UE::ConcertSyncTests::Replication::Frequency
 	}
 	
 	/** Tests that something is sent after the replication interval has expired. */
-	IMPLEMENT_CUSTOM_SIMPLE_AUTOMATION_TEST(FReceiveAfterSpecifiedFrequencyInterval, FSendReceiveObjectWithFrequencyTest, "Editor.Concert.Replication.Frequency.SendingClientRespectsFrequency", EAutomationTestFlags::EditorContext | EAutomationTestFlags::EngineFilter);
+	IMPLEMENT_CUSTOM_SIMPLE_AUTOMATION_TEST(FReceiveAfterSpecifiedFrequencyInterval, FSendReceiveObjectWithFrequencyTest, "VirtualProduction.Concert.Replication.Frequency.SendingClientRespectsFrequency", EAutomationTestFlags::EditorContext | EAutomationTestFlags::EngineFilter);
 	bool FReceiveAfterSpecifiedFrequencyInterval::RunTest(const FString& Parameters)
 	{
 		SetUpClientAndServer();
@@ -76,7 +76,7 @@ namespace UE::ConcertSyncTests::Replication::Frequency
 
 	// TODO UE-203616: This test sometimes fails. Find out why.
 	/** Tests that after about 1 second, the expected number of events was received. */
-	/*IMPLEMENT_CUSTOM_SIMPLE_AUTOMATION_TEST(FSimulateAt30FPS, FSendReceiveObjectWithFrequencyTest, "Editor.Concert.Replication.Frequency.SimulateAt30FPS", EAutomationTestFlags::EditorContext | EAutomationTestFlags::EngineFilter);
+	/*IMPLEMENT_CUSTOM_SIMPLE_AUTOMATION_TEST(FSimulateAt30FPS, FSendReceiveObjectWithFrequencyTest, "VirtualProduction.Concert.Replication.Frequency.SimulateAt30FPS", EAutomationTestFlags::EditorContext | EAutomationTestFlags::EngineFilter);
 	bool FSimulateAt30FPS::RunTest(const FString& Parameters)
 	{
 		// 1. Set up server
diff --git a/Engine/Plugins/Developer/Concert/ConcertSync/ConcertSyncTest/Source/ConcertSyncTest/Private/Replication/Muting/MuteAndQuery.spec.cpp b/Engine/Plugins/Developer/Concert/ConcertSync/ConcertSyncTest/Source/ConcertSyncTest/Private/Replication/Muting/MuteAndQuery.spec.cpp
--- a/Engine/Plugins/Developer/Concert/ConcertSync/ConcertSyncTest/Source/ConcertSyncTest/Private/Replication/Muting/MuteAndQuery.spec.cpp
+++ b/Engine/Plugins/Developer/Concert/ConcertSync/ConcertSyncTest/Source/ConcertSyncTest/Private/Replication/Muting/MuteAndQuery.spec.cpp
@@ -8,7 +8,7 @@
 namespace UE::ConcertSyncTests::Replication
 {
 	/** This tests that muting & unmuting retains a consistent check across various types of changes. */
-	BEGIN_DEFINE_SPEC(FMuteAndQuerySpec, "Editor.Concert.Replication.Muting", EAutomationTestFlags::EditorContext | EAutomationTestFlags::EngineFilter)
+	BEGIN_DEFINE_SPEC(FMuteAndQuerySpec, "VirtualProduction.Concert.Replication.Muting", EAutomationTestFlags::EditorContext | EAutomationTestFlags::EngineFilter)
 	TSharedPtr<FReplicationServer> Server;
 	FReplicationClient* Sender = nullptr;
 
diff --git a/Engine/Plugins/Developer/Concert/ConcertSync/ConcertSyncTest/Source/ConcertSyncTest/Private/Replication/Muting/MuteAndSyncControl.spec.cpp b/Engine/Plugins/Developer/Concert/ConcertSync/ConcertSyncTest/Source/ConcertSyncTest/Private/Replication/Muting/MuteAndSyncControl.spec.cpp
--- a/Engine/Plugins/Developer/Concert/ConcertSync/ConcertSyncTest/Source/ConcertSyncTest/Private/Replication/Muting/MuteAndSyncControl.spec.cpp
+++ b/Engine/Plugins/Developer/Concert/ConcertSync/ConcertSyncTest/Source/ConcertSyncTest/Private/Replication/Muting/MuteAndSyncControl.spec.cpp
@@ -8,7 +8,7 @@
 
 namespace UE::ConcertSyncTests::Replication
 {
-	BEGIN_DEFINE_SPEC(FMuteAndSyncControlSpec, "Editor.Concert.Replication.Muting", EAutomationTestFlags::EditorContext | EAutomationTestFlags::EngineFilter)
+	BEGIN_DEFINE_SPEC(FMuteAndSyncControlSpec, "VirtualProduction.Concert.Replication.Muting", EAutomationTestFlags::EditorContext | EAutomationTestFlags::EngineFilter)
 	TSharedPtr<FReplicationServer> Server;
 	FReplicationClient* Sender = nullptr;
 	FReplicationClient* Receiver = nullptr;
diff --git a/Engine/Plugins/Developer/Concert/ConcertSync/ConcertSyncTest/Source/ConcertSyncTest/Private/Replication/Muting/MutingRejection.spec.cpp b/Engine/Plugins/Developer/Concert/ConcertSync/ConcertSyncTest/Source/ConcertSyncTest/Private/Replication/Muting/MutingRejection.spec.cpp
--- a/Engine/Plugins/Developer/Concert/ConcertSync/ConcertSyncTest/Source/ConcertSyncTest/Private/Replication/Muting/MutingRejection.spec.cpp
+++ b/Engine/Plugins/Developer/Concert/ConcertSync/ConcertSyncTest/Source/ConcertSyncTest/Private/Replication/Muting/MutingRejection.spec.cpp
@@ -8,7 +8,7 @@
 namespace UE::ConcertSyncTests::Replication
 {
 	/** Tests that the rejection cases outlined in FConcertReplication_ChangeMuteState_Request work. */
-	BEGIN_DEFINE_SPEC(FMutingRejectionSpec, "Editor.Concert.Replication.Muting", EAutomationTestFlags::EditorContext | EAutomationTestFlags::EngineFilter)
+	BEGIN_DEFINE_SPEC(FMutingRejectionSpec, "VirtualProduction.Concert.Replication.Muting", EAutomationTestFlags::EditorContext | EAutomationTestFlags::EngineFilter)
 		TSharedPtr<FReplicationServer> Server;
 		FReplicationClient* Sender = nullptr;
 		FReplicationClient* Receiver = nullptr;
@@ -160,7 +160,7 @@ namespace UE::ConcertSyncTests::Replication
 	}
 	
 	/** This tests that muting requests are not served if EConcertSyncSessionFlags::ShouldAllowGlobalMuting is not set. */
-	BEGIN_DEFINE_SPEC(FMutingWithoutFlagSpec, "Editor.Concert.Replication.Muting", EAutomationTestFlags::EditorContext | EAutomationTestFlags::EngineFilter)
+	BEGIN_DEFINE_SPEC(FMutingWithoutFlagSpec, "VirtualProduction.Concert.Replication.Muting", EAutomationTestFlags::EditorContext | EAutomationTestFlags::EngineFilter)
 		TSharedPtr<FReplicationServer> Server;
 		FReplicationClient* Sender = nullptr;
 
diff --git a/Engine/Plugins/Developer/Concert/ConcertSync/ConcertSyncTest/Source/ConcertSyncTest/Private/Replication/Muting/MutingReplication.spec.cpp b/Engine/Plugins/Developer/Concert/ConcertSync/ConcertSyncTest/Source/ConcertSyncTest/Private/Replication/Muting/MutingReplication.spec.cpp
--- a/Engine/Plugins/Developer/Concert/ConcertSync/ConcertSyncTest/Source/ConcertSyncTest/Private/Replication/Muting/MutingReplication.spec.cpp
+++ b/Engine/Plugins/Developer/Concert/ConcertSync/ConcertSyncTest/Source/ConcertSyncTest/Private/Replication/Muting/MutingReplication.spec.cpp
@@ -9,7 +9,7 @@
 namespace UE::ConcertSyncTests::Replication
 {
 	/** This tests that replication works correctly after muting & unmuting. */
-	BEGIN_DEFINE_SPEC(FMutingReplicationSpec, "Editor.Concert.Replication.Muting", EAutomationTestFlags::EditorContext | EAutomationTestFlags::EngineFilter)
+	BEGIN_DEFINE_SPEC(FMutingReplicationSpec, "VirtualProduction.Concert.Replication.Muting", EAutomationTestFlags::EditorContext | EAutomationTestFlags::EngineFilter)
 	TSharedPtr<FReplicationServer> Server;
 	FReplicationClient* Sender = nullptr;
 	FReplicationClient* Receiver = nullptr;
diff --git a/Engine/Plugins/Developer/Concert/ConcertSync/ConcertSyncTest/Source/ConcertSyncTest/Private/Replication/Muting/ResetMuteState.spec.cpp b/Engine/Plugins/Developer/Concert/ConcertSync/ConcertSyncTest/Source/ConcertSyncTest/Private/Replication/Muting/ResetMuteState.spec.cpp
--- a/Engine/Plugins/Developer/Concert/ConcertSync/ConcertSyncTest/Source/ConcertSyncTest/Private/Replication/Muting/ResetMuteState.spec.cpp
+++ b/Engine/Plugins/Developer/Concert/ConcertSync/ConcertSyncTest/Source/ConcertSyncTest/Private/Replication/Muting/ResetMuteState.spec.cpp
@@ -9,7 +9,7 @@
 namespace UE::ConcertSyncTests::Replication
 {
 	/** Tests that EConcertReplicationMuteRequestFlags::ResetMuteState works as intended in conjunction with FConcertReplication_ChangeMuteState_Request. */
-	BEGIN_DEFINE_SPEC(FResetMuteStateSpec, "Editor.Concert.Replication.Muting.Reset", EAutomationTestFlags::EditorContext | EAutomationTestFlags::EngineFilter)
+	BEGIN_DEFINE_SPEC(FResetMuteStateSpec, "VirtualProduction.Concert.Replication.Muting.Reset", EAutomationTestFlags::EditorContext | EAutomationTestFlags::EngineFilter)
 		TUniquePtr<FObjectTestReplicator> ObjectReplicator;
 		TUniquePtr<FObjectTestReplicator> SecondObjectReplicator;
 
diff --git a/Engine/Plugins/Developer/Concert/ConcertSync/ConcertSyncTest/Source/ConcertSyncTest/Private/Replication/PutState/ChangeClientEvent.spec.cpp b/Engine/Plugins/Developer/Concert/ConcertSync/ConcertSyncTest/Source/ConcertSyncTest/Private/Replication/PutState/ChangeClientEvent.spec.cpp
--- a/Engine/Plugins/Developer/Concert/ConcertSync/ConcertSyncTest/Source/ConcertSyncTest/Private/Replication/PutState/ChangeClientEvent.spec.cpp
+++ b/Engine/Plugins/Developer/Concert/ConcertSync/ConcertSyncTest/Source/ConcertSyncTest/Private/Replication/PutState/ChangeClientEvent.spec.cpp
@@ -10,7 +10,7 @@
 
 namespace UE::ConcertSyncTests::Replication::ChangeClients
 {
-	BEGIN_DEFINE_SPEC(FChangeClientEventSpec, "Editor.Concert.Replication.PutState.Event", EAutomationTestFlags::EditorContext | EAutomationTestFlags::EngineFilter)
+	BEGIN_DEFINE_SPEC(FChangeClientEventSpec, "VirtualProduction.Concert.Replication.PutState.Event", EAutomationTestFlags::EditorContext | EAutomationTestFlags::EngineFilter)
 		TUniquePtr<FObjectTestReplicator> ObjectReplicator;
 
 		TUniquePtr<FReplicationServer> Server;
diff --git a/Engine/Plugins/Developer/Concert/ConcertSync/ConcertSyncTest/Source/ConcertSyncTest/Private/Replication/PutState/PutStateActivities.spec.cpp b/Engine/Plugins/Developer/Concert/ConcertSync/ConcertSyncTest/Source/ConcertSyncTest/Private/Replication/PutState/PutStateActivities.spec.cpp
--- a/Engine/Plugins/Developer/Concert/ConcertSync/ConcertSyncTest/Source/ConcertSyncTest/Private/Replication/PutState/PutStateActivities.spec.cpp
+++ b/Engine/Plugins/Developer/Concert/ConcertSync/ConcertSyncTest/Source/ConcertSyncTest/Private/Replication/PutState/PutStateActivities.spec.cpp
@@ -12,7 +12,7 @@
 
 namespace UE::ConcertSyncTests::Replication::ChangeClients
 {
-	BEGIN_DEFINE_SPEC(FPutStateActivitiesSpec, "Editor.Concert.Replication.PutState.Activity", EAutomationTestFlags::EditorContext | EAutomationTestFlags::EngineFilter)
+	BEGIN_DEFINE_SPEC(FPutStateActivitiesSpec, "VirtualProduction.Concert.Replication.PutState.Activity", EAutomationTestFlags::EditorContext | EAutomationTestFlags::EngineFilter)
 		TUniquePtr<FScopedSessionDatabase> SessionDatabase;
 		TSharedPtr<ConcertSyncServer::Replication::IReplicationWorkspace> ReplicationWorkspace;
 		TUniquePtr<FObjectTestReplicator> ObjectReplicator;
diff --git a/Engine/Plugins/Developer/Concert/ConcertSync/ConcertSyncTest/Source/ConcertSyncTest/Private/Replication/PutState/PutStateAndMute.spec.cpp b/Engine/Plugins/Developer/Concert/ConcertSync/ConcertSyncTest/Source/ConcertSyncTest/Private/Replication/PutState/PutStateAndMute.spec.cpp
--- a/Engine/Plugins/Developer/Concert/ConcertSync/ConcertSyncTest/Source/ConcertSyncTest/Private/Replication/PutState/PutStateAndMute.spec.cpp
+++ b/Engine/Plugins/Developer/Concert/ConcertSync/ConcertSyncTest/Source/ConcertSyncTest/Private/Replication/PutState/PutStateAndMute.spec.cpp
@@ -10,7 +10,7 @@
 
 namespace UE::ConcertSyncTests::Replication::ChangeClients
 {
-	BEGIN_DEFINE_SPEC(FPutStateAndMuteSpec, "Editor.Concert.Replication.PutState.Mute", EAutomationTestFlags::EditorContext | EAutomationTestFlags::EngineFilter)
+	BEGIN_DEFINE_SPEC(FPutStateAndMuteSpec, "VirtualProduction.Concert.Replication.PutState.Mute", EAutomationTestFlags::EditorContext | EAutomationTestFlags::EngineFilter)
 		TUniquePtr<FObjectTestReplicator> ObjectReplicator;
 
 		TUniquePtr<FReplicationServer> Server;
diff --git a/Engine/Plugins/Developer/Concert/ConcertSync/ConcertSyncTest/Source/ConcertSyncTest/Private/Replication/PutState/PutStateAndReplicate.spec.cpp b/Engine/Plugins/Developer/Concert/ConcertSync/ConcertSyncTest/Source/ConcertSyncTest/Private/Replication/PutState/PutStateAndReplicate.spec.cpp
--- a/Engine/Plugins/Developer/Concert/ConcertSync/ConcertSyncTest/Source/ConcertSyncTest/Private/Replication/PutState/PutStateAndReplicate.spec.cpp
+++ b/Engine/Plugins/Developer/Concert/ConcertSync/ConcertSyncTest/Source/ConcertSyncTest/Private/Replication/PutState/PutStateAndReplicate.spec.cpp
@@ -9,7 +9,7 @@
 
 namespace UE::ConcertSyncTests::Replication::ChangeClients
 {
-	BEGIN_DEFINE_SPEC(FPutStateAndReplicateSpec, "Editor.Concert.Replication.PutState.Replicate", EAutomationTestFlags::EditorContext | EAutomationTestFlags::EngineFilter)
+	BEGIN_DEFINE_SPEC(FPutStateAndReplicateSpec, "VirtualProduction.Concert.Replication.PutState.Replicate", EAutomationTestFlags::EditorContext | EAutomationTestFlags::EngineFilter)
 		TUniquePtr<FObjectTestReplicator> ObjectReplicator;
 
 		TUniquePtr<FReplicationServer> Server;
diff --git a/Engine/Plugins/Developer/Concert/ConcertSync/ConcertSyncTest/Source/ConcertSyncTest/Private/Replication/PutState/PutStateDisabled.spec.cpp b/Engine/Plugins/Developer/Concert/ConcertSync/ConcertSyncTest/Source/ConcertSyncTest/Private/Replication/PutState/PutStateDisabled.spec.cpp
--- a/Engine/Plugins/Developer/Concert/ConcertSync/ConcertSyncTest/Source/ConcertSyncTest/Private/Replication/PutState/PutStateDisabled.spec.cpp
+++ b/Engine/Plugins/Developer/Concert/ConcertSync/ConcertSyncTest/Source/ConcertSyncTest/Private/Replication/PutState/PutStateDisabled.spec.cpp
@@ -9,7 +9,7 @@
 
 namespace UE::ConcertSyncTests::Replication::ChangeClients
 {
-	BEGIN_DEFINE_SPEC(FPutStateDisabledSpec, "Editor.Concert.Replication.PutState", EAutomationTestFlags::EditorContext | EAutomationTestFlags::EngineFilter)
+	BEGIN_DEFINE_SPEC(FPutStateDisabledSpec, "VirtualProduction.Concert.Replication.PutState", EAutomationTestFlags::EditorContext | EAutomationTestFlags::EngineFilter)
 		const FGuid ClientId { 0, 0, 0, 1 };
 	
 		void TestRequestHasErrorCode(EConcertReplicationPutStateResponseCode ExpectedErrorCode, EConcertSyncSessionFlags Flags, FConcertReplication_PutState_Request Request = {})
diff --git a/Engine/Plugins/Developer/Concert/ConcertSync/ConcertSyncTest/Source/ConcertSyncTest/Private/Replication/PutState/PutStateMisc.spec.cpp b/Engine/Plugins/Developer/Concert/ConcertSync/ConcertSyncTest/Source/ConcertSyncTest/Private/Replication/PutState/PutStateMisc.spec.cpp
--- a/Engine/Plugins/Developer/Concert/ConcertSync/ConcertSyncTest/Source/ConcertSyncTest/Private/Replication/PutState/PutStateMisc.spec.cpp
+++ b/Engine/Plugins/Developer/Concert/ConcertSync/ConcertSyncTest/Source/ConcertSyncTest/Private/Replication/PutState/PutStateMisc.spec.cpp
@@ -11,7 +11,7 @@
 
 namespace UE::ConcertSyncTests::Replication::ChangeClients
 {
-	BEGIN_DEFINE_SPEC(FPutStateMiscSpec, "Editor.Concert.Replication.PutState.Misc", EAutomationTestFlags::EditorContext | EAutomationTestFlags::EngineFilter)
+	BEGIN_DEFINE_SPEC(FPutStateMiscSpec, "VirtualProduction.Concert.Replication.PutState.Misc", EAutomationTestFlags::EditorContext | EAutomationTestFlags::EngineFilter)
 		TUniquePtr<FObjectTestReplicator> ObjectReplicator;
 
 		TUniquePtr<FReplicationServer> Server;
diff --git a/Engine/Plugins/Developer/Concert/ConcertSync/ConcertSyncTest/Source/ConcertSyncTest/Private/Replication/PutState/PutStateRejection.spec.cpp b/Engine/Plugins/Developer/Concert/ConcertSync/ConcertSyncTest/Source/ConcertSyncTest/Private/Replication/PutState/PutStateRejection.spec.cpp
--- a/Engine/Plugins/Developer/Concert/ConcertSync/ConcertSyncTest/Source/ConcertSyncTest/Private/Replication/PutState/PutStateRejection.spec.cpp
+++ b/Engine/Plugins/Developer/Concert/ConcertSync/ConcertSyncTest/Source/ConcertSyncTest/Private/Replication/PutState/PutStateRejection.spec.cpp
@@ -9,7 +9,7 @@
 
 namespace UE::ConcertSyncTests::Replication::ChangeClients
 {
-	BEGIN_DEFINE_SPEC(FPutStateRejectionSpec, "Editor.Concert.Replication.PutState.Rejection", EAutomationTestFlags::EditorContext | EAutomationTestFlags::EngineFilter)
+	BEGIN_DEFINE_SPEC(FPutStateRejectionSpec, "VirtualProduction.Concert.Replication.PutState.Rejection", EAutomationTestFlags::EditorContext | EAutomationTestFlags::EngineFilter)
 		TUniquePtr<FObjectTestReplicator> ObjectReplicator;
 
 		TUniquePtr<FReplicationServer> Server;
diff --git a/Engine/Plugins/Developer/Concert/ConcertSync/ConcertSyncTest/Source/ConcertSyncTest/Private/Replication/PutState/PutStateSyncControl.spec.cpp b/Engine/Plugins/Developer/Concert/ConcertSync/ConcertSyncTest/Source/ConcertSyncTest/Private/Replication/PutState/PutStateSyncControl.spec.cpp
--- a/Engine/Plugins/Developer/Concert/ConcertSync/ConcertSyncTest/Source/ConcertSyncTest/Private/Replication/PutState/PutStateSyncControl.spec.cpp
+++ b/Engine/Plugins/Developer/Concert/ConcertSync/ConcertSyncTest/Source/ConcertSyncTest/Private/Replication/PutState/PutStateSyncControl.spec.cpp
@@ -9,7 +9,7 @@
 
 namespace UE::ConcertSyncTests::Replication::ChangeClients
 {
-	BEGIN_DEFINE_SPEC(FPutStateSyncControlSpec, "Editor.Concert.Replication.PutState.SyncControl", EAutomationTestFlags::EditorContext | EAutomationTestFlags::EngineFilter)
+	BEGIN_DEFINE_SPEC(FPutStateSyncControlSpec, "VirtualProduction.Concert.Replication.PutState.SyncControl", EAutomationTestFlags::EditorContext | EAutomationTestFlags::EngineFilter)
 		TUniquePtr<FObjectTestReplicator> ObjectReplicator;
 
 		TUniquePtr<FReplicationServer> Server;
diff --git a/Engine/Plugins/Developer/Concert/ConcertSync/ConcertSyncTest/Source/ConcertSyncTest/Private/Replication/ReplicationAuthorityTests.cpp b/Engine/Plugins/Developer/Concert/ConcertSync/ConcertSyncTest/Source/ConcertSyncTest/Private/Replication/ReplicationAuthorityTests.cpp
--- a/Engine/Plugins/Developer/Concert/ConcertSync/ConcertSyncTest/Source/ConcertSyncTest/Private/Replication/ReplicationAuthorityTests.cpp
+++ b/Engine/Plugins/Developer/Concert/ConcertSync/ConcertSyncTest/Source/ConcertSyncTest/Private/Replication/ReplicationAuthorityTests.cpp
@@ -19,7 +19,7 @@
 namespace UE::ConcertSyncTests::Replication::Authority
 {
 	/** Data sent from a client that does not have authority over objects is rejected. */
-	IMPLEMENT_CUSTOM_SIMPLE_AUTOMATION_TEST(FRejectUnauthorativeClientTest, FSendReceiveObjectTestBase, "Editor.Concert.Replication.Authority.RejectUnauthorativeClient", EAutomationTestFlags::EditorContext | EAutomationTestFlags::EngineFilter);
+	IMPLEMENT_CUSTOM_SIMPLE_AUTOMATION_TEST(FRejectUnauthorativeClientTest, FSendReceiveObjectTestBase, "VirtualProduction.Concert.Replication.Authority.RejectUnauthorativeClient", EAutomationTestFlags::EditorContext | EAutomationTestFlags::EngineFilter);
 	bool FRejectUnauthorativeClientTest::RunTest(const FString& Parameters)
 	{
 		// 1. Init
@@ -93,7 +93,7 @@ namespace UE::ConcertSyncTests::Replication::Authority
 	}
 
 	/** Makes sure only one client can take authority over the same object properties at the same time and that other clients continue being able to take authority after authority has been released.  */
-	IMPLEMENT_CUSTOM_SIMPLE_AUTOMATION_TEST(FAuthorityIsMutallyExclusiveTest, FSendReceiveGenericTestBase, "Editor.Concert.Replication.Authority.AuthorityIsMutallyExclusive", EAutomationTestFlags::EditorContext | EAutomationTestFlags::EngineFilter);
+	IMPLEMENT_CUSTOM_SIMPLE_AUTOMATION_TEST(FAuthorityIsMutallyExclusiveTest, FSendReceiveGenericTestBase, "VirtualProduction.Concert.Replication.Authority.AuthorityIsMutallyExclusive", EAutomationTestFlags::EditorContext | EAutomationTestFlags::EngineFilter);
 	bool FAuthorityIsMutallyExclusiveTest::RunTest(const FString& Parameters)
 	{
 		// 1. Init
@@ -173,7 +173,7 @@ namespace UE::ConcertSyncTests::Replication::Authority
 	}
 	
 	/** Clients cannot take authority over objects that they did not send */
-	IMPLEMENT_CUSTOM_SIMPLE_AUTOMATION_TEST(FCannotTakeAuthorityOverUnregisteredObjectsTest, FSendReceiveGenericTestBase, "Editor.Concert.Replication.Authority.CannotTakeAuthorityOverUnregisteredObjects", EAutomationTestFlags::EditorContext | EAutomationTestFlags::EngineFilter);
+	IMPLEMENT_CUSTOM_SIMPLE_AUTOMATION_TEST(FCannotTakeAuthorityOverUnregisteredObjectsTest, FSendReceiveGenericTestBase, "VirtualProduction.Concert.Replication.Authority.CannotTakeAuthorityOverUnregisteredObjects", EAutomationTestFlags::EditorContext | EAutomationTestFlags::EngineFilter);
 	bool FCannotTakeAuthorityOverUnregisteredObjectsTest::RunTest(const FString& Parameters)
 	{
 		// FSendReceiveGenericTestBase is used here because it registers no sending streams by default
@@ -215,7 +215,7 @@ namespace UE::ConcertSyncTests::Replication::Authority
 	}
 
 	/** When a client leaves the replication session, the client releases all of the authority as well.  */
-	IMPLEMENT_CUSTOM_SIMPLE_AUTOMATION_TEST(FClientLeavingReplicationSessionLosesAuthority, FSendReceiveGenericTestBase, "Editor.Concert.Replication.Authority.ClientLeavingReplicationSessionLosesAuthority", EAutomationTestFlags::EditorContext | EAutomationTestFlags::EngineFilter);
+	IMPLEMENT_CUSTOM_SIMPLE_AUTOMATION_TEST(FClientLeavingReplicationSessionLosesAuthority, FSendReceiveGenericTestBase, "VirtualProduction.Concert.Replication.Authority.ClientLeavingReplicationSessionLosesAuthority", EAutomationTestFlags::EditorContext | EAutomationTestFlags::EngineFilter);
 	bool FClientLeavingReplicationSessionLosesAuthority::RunTest(const FString& Parameters)
 	{
 		// 1. Init
@@ -261,7 +261,7 @@ namespace UE::ConcertSyncTests::Replication::Authority
 	}
 
 	/** When a client leaves the replication session, the client releases all of the authority as well.  */
-	IMPLEMENT_CUSTOM_SIMPLE_AUTOMATION_TEST(FClientLeavingConcertSessionLosesAuthority, FSendReceiveGenericTestBase, "Editor.Concert.Replication.Authority.ClientLeavingConcertSessionLosesAuthority", EAutomationTestFlags::EditorContext | EAutomationTestFlags::EngineFilter);
+	IMPLEMENT_CUSTOM_SIMPLE_AUTOMATION_TEST(FClientLeavingConcertSessionLosesAuthority, FSendReceiveGenericTestBase, "VirtualProduction.Concert.Replication.Authority.ClientLeavingConcertSessionLosesAuthority", EAutomationTestFlags::EditorContext | EAutomationTestFlags::EngineFilter);
 	bool FClientLeavingConcertSessionLosesAuthority::RunTest(const FString& Parameters)
 	{
 		// 1. Init
@@ -311,7 +311,7 @@ namespace UE::ConcertSyncTests::Replication::Authority
 	 * - Taking authority calls IConcertClientReplicationBridge::PushTrackedObjects
 	 * - Releasing authority calls IConcertClientReplicationBridge::ReleaseTrackedObjects
 	 */
-	IMPLEMENT_CUSTOM_SIMPLE_AUTOMATION_TEST(FChangingAuthorityUpdatesClientBridge, FSendReceiveObjectTestBase, "Editor.Concert.Replication.Authority.ChangingAuthorityUpdatesClientBridge", EAutomationTestFlags::EditorContext | EAutomationTestFlags::EngineFilter);
+	IMPLEMENT_CUSTOM_SIMPLE_AUTOMATION_TEST(FChangingAuthorityUpdatesClientBridge, FSendReceiveObjectTestBase, "VirtualProduction.Concert.Replication.Authority.ChangingAuthorityUpdatesClientBridge", EAutomationTestFlags::EditorContext | EAutomationTestFlags::EngineFilter);
 	bool FChangingAuthorityUpdatesClientBridge::RunTest(const FString& Parameters)
 	{
 		// 1. Init
@@ -353,7 +353,7 @@ namespace UE::ConcertSyncTests::Replication::Authority
 	/**
 	 * Tests that client updates its local cache of the server state when RequestAuthorityChange times out.
 	 */
-	IMPLEMENT_CUSTOM_SIMPLE_AUTOMATION_TEST(FChangingAuthorityTimeoutRetainsServerState, FSendReceiveObjectTestBase, "Editor.Concert.Replication.Authority.ChangeAuthorityTimeoutRetainsServerState", EAutomationTestFlags::EditorContext | EAutomationTestFlags::EngineFilter);
+	IMPLEMENT_CUSTOM_SIMPLE_AUTOMATION_TEST(FChangingAuthorityTimeoutRetainsServerState, FSendReceiveObjectTestBase, "VirtualProduction.Concert.Replication.Authority.ChangeAuthorityTimeoutRetainsServerState", EAutomationTestFlags::EditorContext | EAutomationTestFlags::EngineFilter);
 	bool FChangingAuthorityTimeoutRetainsServerState::RunTest(const FString& Parameters)
 	{
 		// 1. Init
@@ -432,7 +432,7 @@ namespace UE::ConcertSyncTests::Replication::Authority
 	}
 
 	/** Tests UE::ConcertSyncCore::Replication::AuthorityConflictUtils::EnumerateAuthorityConflicts. */
-	IMPLEMENT_CUSTOM_SIMPLE_AUTOMATION_TEST(FConflictEnumerationComponenTests, FSendReceiveObjectTestBase, "Editor.Concert.Replication.Authority.ConflictEnumerationComponenTests", EAutomationTestFlags::EditorContext | EAutomationTestFlags::EngineFilter);
+	IMPLEMENT_CUSTOM_SIMPLE_AUTOMATION_TEST(FConflictEnumerationComponenTests, FSendReceiveObjectTestBase, "VirtualProduction.Concert.Replication.Authority.ConflictEnumerationComponenTests", EAutomationTestFlags::EditorContext | EAutomationTestFlags::EngineFilter);
 	bool FConflictEnumerationComponenTests::RunTest(const FString& Parameters)
 	{
 		using namespace ConflictEnumerationComponent;
diff --git a/Engine/Plugins/Developer/Concert/ConcertSync/ConcertSyncTest/Source/ConcertSyncTest/Private/Replication/ReplicationHandshakeTests.cpp b/Engine/Plugins/Developer/Concert/ConcertSync/ConcertSyncTest/Source/ConcertSyncTest/Private/Replication/ReplicationHandshakeTests.cpp
--- a/Engine/Plugins/Developer/Concert/ConcertSync/ConcertSyncTest/Source/ConcertSyncTest/Private/Replication/ReplicationHandshakeTests.cpp
+++ b/Engine/Plugins/Developer/Concert/ConcertSync/ConcertSyncTest/Source/ConcertSyncTest/Private/Replication/ReplicationHandshakeTests.cpp
@@ -21,7 +21,7 @@ namespace UE::ConcertSyncTests::Replication::Handshake
 	/**
 	 * Tests the handshake for joining and leaving a replication session.
 	 */
-	IMPLEMENT_CUSTOM_SIMPLE_AUTOMATION_TEST(FJoinHandshakeTest, FConcertClientServerCommunicationTest, "Editor.Concert.Replication.Handshake", EAutomationTestFlags::EditorContext | EAutomationTestFlags::EngineFilter);
+	IMPLEMENT_CUSTOM_SIMPLE_AUTOMATION_TEST(FJoinHandshakeTest, FConcertClientServerCommunicationTest, "VirtualProduction.Concert.Replication.Handshake", EAutomationTestFlags::EditorContext | EAutomationTestFlags::EngineFilter);
 	bool FJoinHandshakeTest::RunTest(const FString& Parameters)
 	{
 		using namespace ConcertSyncClient::TestInterface;
diff --git a/Engine/Plugins/Developer/Concert/ConcertSync/ConcertSyncTest/Source/ConcertSyncTest/Private/Replication/ReplicationSendReceiveTests.cpp b/Engine/Plugins/Developer/Concert/ConcertSync/ConcertSyncTest/Source/ConcertSyncTest/Private/Replication/ReplicationSendReceiveTests.cpp
--- a/Engine/Plugins/Developer/Concert/ConcertSync/ConcertSyncTest/Source/ConcertSyncTest/Private/Replication/ReplicationSendReceiveTests.cpp
+++ b/Engine/Plugins/Developer/Concert/ConcertSync/ConcertSyncTest/Source/ConcertSyncTest/Private/Replication/ReplicationSendReceiveTests.cpp
@@ -93,7 +93,7 @@ namespace UE::ConcertSyncTests::Replication::SendReceiveFlow
 	}
 	
 	/** Tests replicating data from sender client > server > receiver client. */
-	IMPLEMENT_CUSTOM_SIMPLE_AUTOMATION_TEST(FSendReceiveFlowTests, FSendReceiveObjectTestBase, "Editor.Concert.Replication.SendReceive.SendValuesOnSingleStream", EAutomationTestFlags::EditorContext | EAutomationTestFlags::EngineFilter);
+	IMPLEMENT_CUSTOM_SIMPLE_AUTOMATION_TEST(FSendReceiveFlowTests, FSendReceiveObjectTestBase, "VirtualProduction.Concert.Replication.SendReceive.SendValuesOnSingleStream", EAutomationTestFlags::EditorContext | EAutomationTestFlags::EngineFilter);
 	bool FSendReceiveFlowTests::RunTest(const FString& Parameters)
 	{
 		SharedSetupSet(*this);
@@ -102,7 +102,7 @@ namespace UE::ConcertSyncTests::Replication::SendReceiveFlow
 	}
 	
 	/** Tests replicating data from sender client > server > receiver client. */
-	IMPLEMENT_CUSTOM_SIMPLE_AUTOMATION_TEST(FSendCDOValuesOnSingleStream, FSendReceiveObjectTestBase, "Editor.Concert.Replication.SendReceive.SendCDOValuesOnSingleStream", EAutomationTestFlags::EditorContext | EAutomationTestFlags::EngineFilter);
+	IMPLEMENT_CUSTOM_SIMPLE_AUTOMATION_TEST(FSendCDOValuesOnSingleStream, FSendReceiveObjectTestBase, "VirtualProduction.Concert.Replication.SendReceive.SendCDOValuesOnSingleStream", EAutomationTestFlags::EditorContext | EAutomationTestFlags::EngineFilter);
 	bool FSendCDOValuesOnSingleStream::RunTest(const FString& Parameters)
 	{
 		SharedSetupSet(*this);
@@ -156,7 +156,7 @@ namespace UE::ConcertSyncTests::Replication::SendReceiveFlow
 	};
 
 	/** Test which still sends TestObject but does so with two separate streams: one for the float and the other for the vector property. */
-	IMPLEMENT_CUSTOM_SIMPLE_AUTOMATION_TEST(FSendReceiveMultiStreamSameObjectTests, FSplitSendReceiveObjectTest, "Editor.Concert.Replication.SendReceive.MultipleStreamsForSameObject", EAutomationTestFlags::EditorContext | EAutomationTestFlags::EngineFilter);
+	IMPLEMENT_CUSTOM_SIMPLE_AUTOMATION_TEST(FSendReceiveMultiStreamSameObjectTests, FSplitSendReceiveObjectTest, "VirtualProduction.Concert.Replication.SendReceive.MultipleStreamsForSameObject", EAutomationTestFlags::EditorContext | EAutomationTestFlags::EngineFilter);
 	bool FSendReceiveMultiStreamSameObjectTests::RunTest(const FString& Parameters)
 	{
 		SharedSetupSet(*this);
@@ -165,7 +165,7 @@ namespace UE::ConcertSyncTests::Replication::SendReceiveFlow
 	}
 
 	/** Tests that changing a stream while replication is in progress actually changes the properties being replicated. */
-	IMPLEMENT_CUSTOM_SIMPLE_AUTOMATION_TEST(FUpdateStreamWhileSendingUpdatesSentProperties, FSendReceiveObjectTestBase, "Editor.Concert.Replication.SendReceive.UpdateStreamWhileSendingUpdatesSentProperties", EAutomationTestFlags::EditorContext | EAutomationTestFlags::EngineFilter);
+	IMPLEMENT_CUSTOM_SIMPLE_AUTOMATION_TEST(FUpdateStreamWhileSendingUpdatesSentProperties, FSendReceiveObjectTestBase, "VirtualProduction.Concert.Replication.SendReceive.UpdateStreamWhileSendingUpdatesSentProperties", EAutomationTestFlags::EditorContext | EAutomationTestFlags::EngineFilter);
 	bool FUpdateStreamWhileSendingUpdatesSentProperties::RunTest(const FString& Parameters)
 	{
 		// 1. Send all properties with the stream from handshake
diff --git a/Engine/Plugins/Developer/Concert/ConcertSync/ConcertSyncTest/Source/ConcertSyncTest/Private/Replication/ReplicationWorkspace.spec.cpp b/Engine/Plugins/Developer/Concert/ConcertSync/ConcertSyncTest/Source/ConcertSyncTest/Private/Replication/ReplicationWorkspace.spec.cpp
--- a/Engine/Plugins/Developer/Concert/ConcertSync/ConcertSyncTest/Source/ConcertSyncTest/Private/Replication/ReplicationWorkspace.spec.cpp
+++ b/Engine/Plugins/Developer/Concert/ConcertSync/ConcertSyncTest/Source/ConcertSyncTest/Private/Replication/ReplicationWorkspace.spec.cpp
@@ -13,7 +13,7 @@
 
 namespace UE::ConcertSyncTests::Replication::UI
 {
-	BEGIN_DEFINE_SPEC(FReplicationWorkspaceSpec, "Editor.Concert.Replication.ReplicationWorkspace", EAutomationTestFlags::EditorContext | EAutomationTestFlags::EngineFilter)
+	BEGIN_DEFINE_SPEC(FReplicationWorkspaceSpec, "VirtualProduction.Concert.Replication.ReplicationWorkspace", EAutomationTestFlags::EditorContext | EAutomationTestFlags::EngineFilter)
 		TUniquePtr<FScopedSessionDatabase> SessionDatabase;
 		TSharedPtr<ConcertSyncServer::Replication::IReplicationWorkspace> ReplicationWorkspace;
 	
diff --git a/Engine/Plugins/Developer/Concert/ConcertSync/ConcertSyncTest/Source/ConcertSyncTest/Private/Replication/Restore/ActivityProduction.spec.cpp b/Engine/Plugins/Developer/Concert/ConcertSync/ConcertSyncTest/Source/ConcertSyncTest/Private/Replication/Restore/ActivityProduction.spec.cpp
--- a/Engine/Plugins/Developer/Concert/ConcertSync/ConcertSyncTest/Source/ConcertSyncTest/Private/Replication/Restore/ActivityProduction.spec.cpp
+++ b/Engine/Plugins/Developer/Concert/ConcertSync/ConcertSyncTest/Source/ConcertSyncTest/Private/Replication/Restore/ActivityProduction.spec.cpp
@@ -1,4 +1,4 @@
-﻿// Copyright Epic Games, Inc. All Rights Reserved.
+// Copyright Epic Games, Inc. All Rights Reserved.
 
 #include "Replication/Util/Mocks/ReplicationWorkspaceCallInterceptorMock.h"
 #include "Replication/Util/Spec/ReplicationClient.h"
@@ -10,7 +10,7 @@
 
 namespace UE::ConcertSyncTests::Replication
 {
-	BEGIN_DEFINE_SPEC(FLeaveReplicationActivitySpec, "Editor.Concert.Replication.RestoreContent.Activity", EAutomationTestFlags::EditorContext | EAutomationTestFlags::EngineFilter)
+	BEGIN_DEFINE_SPEC(FLeaveReplicationActivitySpec, "VirtualProduction.Concert.Replication.RestoreContent.Activity", EAutomationTestFlags::EditorContext | EAutomationTestFlags::EngineFilter)
 		/** Detects calls into the workspace */
 		TSharedPtr<FReplicationWorkspaceCallInterceptorMock> WorkspaceMock;
 	
@@ -92,7 +92,7 @@ namespace UE::ConcertSyncTests::Replication
 	}
 
 
-	BEGIN_DEFINE_SPEC(FMuteReplicationActivitySpec, "Editor.Concert.Replication.RestoreContent.Activity", EAutomationTestFlags::EditorContext | EAutomationTestFlags::EngineFilter)
+	BEGIN_DEFINE_SPEC(FMuteReplicationActivitySpec, "VirtualProduction.Concert.Replication.RestoreContent.Activity", EAutomationTestFlags::EditorContext | EAutomationTestFlags::EngineFilter)
 		/** Detects calls into the workspace */
 		TSharedPtr<FReplicationWorkspaceCallInterceptorMock> WorkspaceMock;
 	
@@ -149,7 +149,7 @@ namespace UE::ConcertSyncTests::Replication
 		});
 	}
 	
-	BEGIN_DEFINE_SPEC(FNoReplicationActivitiesSpec, "Editor.Concert.Replication.RestoreContent.Activity", EAutomationTestFlags::EditorContext | EAutomationTestFlags::EngineFilter)
+	BEGIN_DEFINE_SPEC(FNoReplicationActivitiesSpec, "VirtualProduction.Concert.Replication.RestoreContent.Activity", EAutomationTestFlags::EditorContext | EAutomationTestFlags::EngineFilter)
 		/** Detects calls into the workspace */
 		TSharedPtr<FReplicationWorkspaceCallInterceptorMock> WorkspaceMock;
 		TUniquePtr<FObjectTestReplicator> ObjectReplicator;
diff --git a/Engine/Plugins/Developer/Concert/ConcertSync/ConcertSyncTest/Source/ConcertSyncTest/Private/Replication/Restore/RestoreContentDisabled.spec.cpp b/Engine/Plugins/Developer/Concert/ConcertSync/ConcertSyncTest/Source/ConcertSyncTest/Private/Replication/Restore/RestoreContentDisabled.spec.cpp
--- a/Engine/Plugins/Developer/Concert/ConcertSync/ConcertSyncTest/Source/ConcertSyncTest/Private/Replication/Restore/RestoreContentDisabled.spec.cpp
+++ b/Engine/Plugins/Developer/Concert/ConcertSync/ConcertSyncTest/Source/ConcertSyncTest/Private/Replication/Restore/RestoreContentDisabled.spec.cpp
@@ -9,7 +9,7 @@
 
 namespace UE::ConcertSyncTests::Replication::RestoreContent
 {
-	BEGIN_DEFINE_SPEC(FRestoreContentDisabledSpec, "Editor.Concert.Replication.RestoreContent", EAutomationTestFlags::EditorContext | EAutomationTestFlags::EngineFilter)
+	BEGIN_DEFINE_SPEC(FRestoreContentDisabledSpec, "VirtualProduction.Concert.Replication.RestoreContent", EAutomationTestFlags::EditorContext | EAutomationTestFlags::EngineFilter)
 		/** Detects calls into the workspace */
 		TSharedPtr<FReplicationWorkspaceCallInterceptorMock> WorkspaceMock;
 	
diff --git a/Engine/Plugins/Developer/Concert/ConcertSync/ConcertSyncTest/Source/ConcertSyncTest/Private/Replication/Restore/RestoreContentPreexisting.spec.cpp b/Engine/Plugins/Developer/Concert/ConcertSync/ConcertSyncTest/Source/ConcertSyncTest/Private/Replication/Restore/RestoreContentPreexisting.spec.cpp
--- a/Engine/Plugins/Developer/Concert/ConcertSync/ConcertSyncTest/Source/ConcertSyncTest/Private/Replication/Restore/RestoreContentPreexisting.spec.cpp
+++ b/Engine/Plugins/Developer/Concert/ConcertSync/ConcertSyncTest/Source/ConcertSyncTest/Private/Replication/Restore/RestoreContentPreexisting.spec.cpp
@@ -11,7 +11,7 @@
 
 namespace UE::ConcertSyncTests::Replication::RestoreContent
 {
-	BEGIN_DEFINE_SPEC(FRestoreContentPreexistingSpec, "Editor.Concert.Replication.RestoreContent", EAutomationTestFlags::EditorContext | EAutomationTestFlags::EngineFilter)
+	BEGIN_DEFINE_SPEC(FRestoreContentPreexistingSpec, "VirtualProduction.Concert.Replication.RestoreContent", EAutomationTestFlags::EditorContext | EAutomationTestFlags::EngineFilter)
 		/** Detects calls into the workspace */
 		TSharedPtr<FReplicationWorkspaceCallInterceptorMock> WorkspaceMock;
 		FConcertSyncReplicationPayload_LeaveReplication LeaveReplicationData;
diff --git a/Engine/Plugins/Developer/Concert/ConcertSync/ConcertSyncTest/Source/ConcertSyncTest/Private/Replication/Restore/RestoreContentWithNoPreexisting.spec.cpp b/Engine/Plugins/Developer/Concert/ConcertSync/ConcertSyncTest/Source/ConcertSyncTest/Private/Replication/Restore/RestoreContentWithNoPreexisting.spec.cpp
--- a/Engine/Plugins/Developer/Concert/ConcertSync/ConcertSyncTest/Source/ConcertSyncTest/Private/Replication/Restore/RestoreContentWithNoPreexisting.spec.cpp
+++ b/Engine/Plugins/Developer/Concert/ConcertSync/ConcertSyncTest/Source/ConcertSyncTest/Private/Replication/Restore/RestoreContentWithNoPreexisting.spec.cpp
@@ -19,7 +19,7 @@ namespace UE::ConcertSyncTests::Replication::RestoreContent
 	ENUM_CLASS_FLAGS(ERestoreTestFlags)
 
 	/** Test that restoring content works correctly when the restoring client has no streams nor authority pre-assigned. */
-	BEGIN_DEFINE_SPEC(FRestoreContentWithNoPreExistingSpec, "Editor.Concert.Replication.RestoreContent", EAutomationTestFlags::EditorContext | EAutomationTestFlags::EngineFilter)
+	BEGIN_DEFINE_SPEC(FRestoreContentWithNoPreExistingSpec, "VirtualProduction.Concert.Replication.RestoreContent", EAutomationTestFlags::EditorContext | EAutomationTestFlags::EngineFilter)
 		/** Detects calls into the workspace */
 		TSharedPtr<FReplicationWorkspaceCallInterceptorMock> WorkspaceMock;
 		FConcertSyncReplicationPayload_LeaveReplication LeaveReplicationData;
@@ -392,7 +392,7 @@ namespace UE::ConcertSyncTests::Replication::RestoreContent
 		});
 	}
 	
-	BEGIN_DEFINE_SPEC(FRestoreContentPropertiesSpec, "Editor.Concert.Replication.RestoreContent", EAutomationTestFlags::EditorContext | EAutomationTestFlags::EngineFilter)
+	BEGIN_DEFINE_SPEC(FRestoreContentPropertiesSpec, "VirtualProduction.Concert.Replication.RestoreContent", EAutomationTestFlags::EditorContext | EAutomationTestFlags::EngineFilter)
 		/** Detects calls into the workspace */
 		TSharedPtr<FReplicationWorkspaceCallInterceptorMock> WorkspaceMock;
 		FConcertSyncReplicationPayload_LeaveReplication LeaveReplicationData;
diff --git a/Engine/Plugins/Developer/Concert/ConcertSync/ConcertSyncTest/Source/ConcertSyncTest/Private/Replication/Restore/RestoreMuteState.spec.cpp b/Engine/Plugins/Developer/Concert/ConcertSync/ConcertSyncTest/Source/ConcertSyncTest/Private/Replication/Restore/RestoreMuteState.spec.cpp
--- a/Engine/Plugins/Developer/Concert/ConcertSync/ConcertSyncTest/Source/ConcertSyncTest/Private/Replication/Restore/RestoreMuteState.spec.cpp
+++ b/Engine/Plugins/Developer/Concert/ConcertSync/ConcertSyncTest/Source/ConcertSyncTest/Private/Replication/Restore/RestoreMuteState.spec.cpp
@@ -11,7 +11,7 @@
 
 namespace UE::ConcertSyncTests::Replication::RestoreContent
 {
-	BEGIN_DEFINE_SPEC(FRestoreMuteStateSpec, "Editor.Concert.Replication.RestoreContent.Mute", EAutomationTestFlags::EditorContext | EAutomationTestFlags::EngineFilter)
+	BEGIN_DEFINE_SPEC(FRestoreMuteStateSpec, "VirtualProduction.Concert.Replication.RestoreContent.Mute", EAutomationTestFlags::EditorContext | EAutomationTestFlags::EngineFilter)
 		TUniquePtr<FScopedSessionDatabase> SessionDatabase;
 		TSharedPtr<ConcertSyncServer::Replication::IReplicationWorkspace> ReplicationWorkspace;
 		TUniquePtr<FObjectTestReplicator> ObjectReplicator;
diff --git a/Engine/Plugins/Developer/Concert/ConcertSync/ConcertSyncTest/Source/ConcertSyncTest/Private/Replication/Restore/RestoreMuteStateSyncControl.spec.cpp b/Engine/Plugins/Developer/Concert/ConcertSync/ConcertSyncTest/Source/ConcertSyncTest/Private/Replication/Restore/RestoreMuteStateSyncControl.spec.cpp
--- a/Engine/Plugins/Developer/Concert/ConcertSync/ConcertSyncTest/Source/ConcertSyncTest/Private/Replication/Restore/RestoreMuteStateSyncControl.spec.cpp
+++ b/Engine/Plugins/Developer/Concert/ConcertSync/ConcertSyncTest/Source/ConcertSyncTest/Private/Replication/Restore/RestoreMuteStateSyncControl.spec.cpp
@@ -12,7 +12,7 @@
 
 namespace UE::ConcertSyncTests::Replication::RestoreContent
 {
-	BEGIN_DEFINE_SPEC(FRestoreMuteStateSyncControlSpec, "Editor.Concert.Replication.RestoreContent.Mute", EAutomationTestFlags::EditorContext | EAutomationTestFlags::EngineFilter)
+	BEGIN_DEFINE_SPEC(FRestoreMuteStateSyncControlSpec, "VirtualProduction.Concert.Replication.RestoreContent.Mute", EAutomationTestFlags::EditorContext | EAutomationTestFlags::EngineFilter)
 		TUniquePtr<FScopedSessionDatabase> SessionDatabase;
 		TSharedPtr<ConcertSyncServer::Replication::IReplicationWorkspace> ReplicationWorkspace;
 		TUniquePtr<FObjectTestReplicator> ObjectReplicator;
diff --git a/Engine/Plugins/Developer/Concert/ConcertSync/ConcertSyncTest/Source/ConcertSyncTest/Private/Replication/Stream/ChangeStream.spec.cpp b/Engine/Plugins/Developer/Concert/ConcertSync/ConcertSyncTest/Source/ConcertSyncTest/Private/Replication/Stream/ChangeStream.spec.cpp
--- a/Engine/Plugins/Developer/Concert/ConcertSync/ConcertSyncTest/Source/ConcertSyncTest/Private/Replication/Stream/ChangeStream.spec.cpp
+++ b/Engine/Plugins/Developer/Concert/ConcertSync/ConcertSyncTest/Source/ConcertSyncTest/Private/Replication/Stream/ChangeStream.spec.cpp
@@ -11,7 +11,7 @@
 
 namespace UE::ConcertSyncTests::Replication::RestoreContent
 {
-	BEGIN_DEFINE_SPEC(FChangeStreamSpec, "Editor.Concert.Replication.Stream", EAutomationTestFlags::EditorContext | EAutomationTestFlags::EngineFilter)
+	BEGIN_DEFINE_SPEC(FChangeStreamSpec, "VirtualProduction.Concert.Replication.Stream", EAutomationTestFlags::EditorContext | EAutomationTestFlags::EngineFilter)
 		TUniquePtr<FObjectTestReplicator> FooObject;
 		TUniquePtr<FObjectTestReplicator> BarObject;
 	
diff --git a/Engine/Plugins/Developer/Concert/ConcertSync/ConcertSyncTest/Source/ConcertSyncTest/Private/Replication/Stream/StreamAuthorityInterplay.spec.cpp b/Engine/Plugins/Developer/Concert/ConcertSync/ConcertSyncTest/Source/ConcertSyncTest/Private/Replication/Stream/StreamAuthorityInterplay.spec.cpp
--- a/Engine/Plugins/Developer/Concert/ConcertSync/ConcertSyncTest/Source/ConcertSyncTest/Private/Replication/Stream/StreamAuthorityInterplay.spec.cpp
+++ b/Engine/Plugins/Developer/Concert/ConcertSync/ConcertSyncTest/Source/ConcertSyncTest/Private/Replication/Stream/StreamAuthorityInterplay.spec.cpp
@@ -1,4 +1,4 @@
-﻿// Copyright Epic Games, Inc. All Rights Reserved.
+// Copyright Epic Games, Inc. All Rights Reserved.
 
 #include "Replication/Util/Mocks/ReplicationWorkspaceCallInterceptorMock.h"
 #include "Replication/Util/Spec/ReplicationClient.h"
@@ -11,7 +11,7 @@
 
 namespace UE::ConcertSyncTests::Replication::RestoreContent
 {
-	BEGIN_DEFINE_SPEC(FStreamAuthorityInterplaySpec, "Editor.Concert.Replication.Stream", EAutomationTestFlags::EditorContext | EAutomationTestFlags::EngineFilter)
+	BEGIN_DEFINE_SPEC(FStreamAuthorityInterplaySpec, "VirtualProduction.Concert.Replication.Stream", EAutomationTestFlags::EditorContext | EAutomationTestFlags::EngineFilter)
 		TUniquePtr<FObjectTestReplicator> ObjectReplicator;
 	
 		TUniquePtr<FReplicationServer> Server;
diff --git a/Engine/Plugins/Developer/Concert/ConcertSync/ConcertSyncTest/Source/ConcertSyncTest/Private/Replication/Stream/StreamRequestTests_Change.cpp b/Engine/Plugins/Developer/Concert/ConcertSync/ConcertSyncTest/Source/ConcertSyncTest/Private/Replication/Stream/StreamRequestTests_Change.cpp
--- a/Engine/Plugins/Developer/Concert/ConcertSync/ConcertSyncTest/Source/ConcertSyncTest/Private/Replication/Stream/StreamRequestTests_Change.cpp
+++ b/Engine/Plugins/Developer/Concert/ConcertSync/ConcertSyncTest/Source/ConcertSyncTest/Private/Replication/Stream/StreamRequestTests_Change.cpp
@@ -15,7 +15,7 @@
 namespace UE::ConcertSyncTests::Replication::Stream
 {
 	/** Simple case of FConcertChangeStream_Request where the requesting client has authority and thus cannot generate any authority conflicts. */
-	IMPLEMENT_CUSTOM_SIMPLE_AUTOMATION_TEST(FSimpleChangeStream, FChangeStreamsTestBase, "Editor.Concert.Replication.Stream.SimpleChangeStream", EAutomationTestFlags::EditorContext | EAutomationTestFlags::EngineFilter);
+	IMPLEMENT_CUSTOM_SIMPLE_AUTOMATION_TEST(FSimpleChangeStream, FChangeStreamsTestBase, "VirtualProduction.Concert.Replication.Stream.SimpleChangeStream", EAutomationTestFlags::EditorContext | EAutomationTestFlags::EngineFilter);
 	bool FSimpleChangeStream::RunTest(const FString& Parameters)
 	{
 		using namespace ConcertSyncClient::Replication;
@@ -53,7 +53,7 @@ namespace UE::ConcertSyncTests::Replication::Stream
 	}
 
 	/** Case of FConcertChangeStream_Request in which requester has authority over an object and creates two streams which have authority over overlapping properties. This is allowed. */
-	IMPLEMENT_CUSTOM_SIMPLE_AUTOMATION_TEST(FChangeStreamWithOverlappingProperties, FChangeStreamsTestBase, "Editor.Concert.Replication.Stream.ChangeStreamWithOverlappingProperties", EAutomationTestFlags::EditorContext | EAutomationTestFlags::EngineFilter);
+	IMPLEMENT_CUSTOM_SIMPLE_AUTOMATION_TEST(FChangeStreamWithOverlappingProperties, FChangeStreamsTestBase, "VirtualProduction.Concert.Replication.Stream.ChangeStreamWithOverlappingProperties", EAutomationTestFlags::EditorContext | EAutomationTestFlags::EngineFilter);
 	bool FChangeStreamWithOverlappingProperties::RunTest(const FString& Parameters)
 	{
 		using namespace ConcertSyncClient::Replication;
@@ -100,7 +100,7 @@ namespace UE::ConcertSyncTests::Replication::Stream
 	 * The requester and client have authority over separate properties on the same object.
 	 * The request would cause the requester's existing authority to overlap with the other client, which is rejected.
 	 */
-	IMPLEMENT_CUSTOM_SIMPLE_AUTOMATION_TEST(FRejectChangeStreamWithConflictingAuthority, FChangeStreamsTestBase, "Editor.Concert.Replication.Stream.RejectChangeStreamWithConflictingAuthority", EAutomationTestFlags::EditorContext | EAutomationTestFlags::EngineFilter);
+	IMPLEMENT_CUSTOM_SIMPLE_AUTOMATION_TEST(FRejectChangeStreamWithConflictingAuthority, FChangeStreamsTestBase, "VirtualProduction.Concert.Replication.Stream.RejectChangeStreamWithConflictingAuthority", EAutomationTestFlags::EditorContext | EAutomationTestFlags::EngineFilter);
 	bool FRejectChangeStreamWithConflictingAuthority::RunTest(const FString& Parameters)
 	{
 		using namespace ConcertSyncClient::Replication;
@@ -186,7 +186,7 @@ namespace UE::ConcertSyncTests::Replication::Stream
 	}
 
 	/** Tests that no changes are made if a sub-step of FConcertChangeStream_Request causes an error. */
-	IMPLEMENT_CUSTOM_SIMPLE_AUTOMATION_TEST(FChangeStreamIsAtmoic, FChangeStreamsTestBase, "Editor.Concert.Replication.Stream.ChangeStreamIsAtmoic", EAutomationTestFlags::EditorContext | EAutomationTestFlags::EngineFilter);
+	IMPLEMENT_CUSTOM_SIMPLE_AUTOMATION_TEST(FChangeStreamIsAtmoic, FChangeStreamsTestBase, "VirtualProduction.Concert.Replication.Stream.ChangeStreamIsAtmoic", EAutomationTestFlags::EditorContext | EAutomationTestFlags::EngineFilter);
 	bool FChangeStreamIsAtmoic::RunTest(const FString& Parameters)
 	{
 		using namespace ConcertSyncClient::Replication;
@@ -227,7 +227,7 @@ namespace UE::ConcertSyncTests::Replication::Stream
 	}
 
 	/** Removing an object from a stream over which the requester has authority, clears the authority the client had. Tests that another client can now claim authority over the previously authored properties. */
-	IMPLEMENT_CUSTOM_SIMPLE_AUTOMATION_TEST(FRemovingObjectClearsAuthority, FChangeStreamsTestBase, "Editor.Concert.Replication.Stream.RemovingObjectClearsAuthority", EAutomationTestFlags::EditorContext | EAutomationTestFlags::EngineFilter);
+	IMPLEMENT_CUSTOM_SIMPLE_AUTOMATION_TEST(FRemovingObjectClearsAuthority, FChangeStreamsTestBase, "VirtualProduction.Concert.Replication.Stream.RemovingObjectClearsAuthority", EAutomationTestFlags::EditorContext | EAutomationTestFlags::EngineFilter);
 	bool FRemovingObjectClearsAuthority::RunTest(const FString& Parameters)
 	{
 		using namespace ConcertSyncClient::Replication;
@@ -267,7 +267,7 @@ namespace UE::ConcertSyncTests::Replication::Stream
 	/**
      * Tests that client updates its local cache of the server state when RequestAuthorityChange times out.
      */
-    IMPLEMENT_CUSTOM_SIMPLE_AUTOMATION_TEST(FChangingStreamTimeoutRetainsServerState, FChangeStreamsTestBase, "Editor.Concert.Replication.Stream.ChangeStreamTimeoutRetainsServerState", EAutomationTestFlags::EditorContext | EAutomationTestFlags::EngineFilter);
+    IMPLEMENT_CUSTOM_SIMPLE_AUTOMATION_TEST(FChangingStreamTimeoutRetainsServerState, FChangeStreamsTestBase, "VirtualProduction.Concert.Replication.Stream.ChangeStreamTimeoutRetainsServerState", EAutomationTestFlags::EditorContext | EAutomationTestFlags::EngineFilter);
     bool FChangingStreamTimeoutRetainsServerState::RunTest(const FString& Parameters)
     {
     	// 1. Start up client with float property stream
diff --git a/Engine/Plugins/Developer/Concert/ConcertSync/ConcertSyncTest/Source/ConcertSyncTest/Private/Replication/Stream/StreamRequestTests_Frequency.cpp b/Engine/Plugins/Developer/Concert/ConcertSync/ConcertSyncTest/Source/ConcertSyncTest/Private/Replication/Stream/StreamRequestTests_Frequency.cpp
--- a/Engine/Plugins/Developer/Concert/ConcertSync/ConcertSyncTest/Source/ConcertSyncTest/Private/Replication/Stream/StreamRequestTests_Frequency.cpp
+++ b/Engine/Plugins/Developer/Concert/ConcertSync/ConcertSyncTest/Source/ConcertSyncTest/Private/Replication/Stream/StreamRequestTests_Frequency.cpp
@@ -22,7 +22,7 @@ namespace UE::ConcertSyncTests::Replication::Stream
 	 * - frequency can only be overriden if object is in object replication map
 	 * - frequency override is removed if object is removed from object replication map
 	 */
-	IMPLEMENT_CUSTOM_SIMPLE_AUTOMATION_TEST(FChangeFrequency, FChangeStreamsTestBase, "Editor.Concert.Replication.Stream.Frequency.ChangeFrequency", EAutomationTestFlags::EditorContext | EAutomationTestFlags::EngineFilter);
+	IMPLEMENT_CUSTOM_SIMPLE_AUTOMATION_TEST(FChangeFrequency, FChangeStreamsTestBase, "VirtualProduction.Concert.Replication.Stream.Frequency.ChangeFrequency", EAutomationTestFlags::EditorContext | EAutomationTestFlags::EngineFilter);
 	bool FChangeFrequency::RunTest(const FString& Parameters)
 	{
 		// 1. Start up client with float property stream
@@ -123,7 +123,7 @@ namespace UE::ConcertSyncTests::Replication::Stream
 	 * - Cannot change frequency for unknown stream
 	 * - Cannot change frequency for object not registerd to stream
 	 */
-	IMPLEMENT_CUSTOM_SIMPLE_AUTOMATION_TEST(FPutAddAndRemoveFrequencyFlow, FChangeStreamsTestBase, "Editor.Concert.Replication.Stream.Frequency.PutAddAndRemoveFrequencyFlow", EAutomationTestFlags::EditorContext | EAutomationTestFlags::EngineFilter);
+	IMPLEMENT_CUSTOM_SIMPLE_AUTOMATION_TEST(FPutAddAndRemoveFrequencyFlow, FChangeStreamsTestBase, "VirtualProduction.Concert.Replication.Stream.Frequency.PutAddAndRemoveFrequencyFlow", EAutomationTestFlags::EditorContext | EAutomationTestFlags::EngineFilter);
 	bool FPutAddAndRemoveFrequencyFlow::RunTest(const FString& Parameters)
 	{
 		UTestReflectionObject* MainTestObject = NewObject<UTestReflectionObject>(GetTransientPackage());
@@ -172,7 +172,7 @@ namespace UE::ConcertSyncTests::Replication::Stream
 	 * - Cannot change frequency for unknown stream
 	 * - Cannot change frequency for object not registerd to stream
 	 */
-	IMPLEMENT_CUSTOM_SIMPLE_AUTOMATION_TEST(FRemovingObjectRemovesFrequencyEntry, FChangeStreamsTestBase, "Editor.Concert.Replication.Stream.Frequency.RemovingObjectRemovesFrequencyEntry", EAutomationTestFlags::EditorContext | EAutomationTestFlags::EngineFilter);
+	IMPLEMENT_CUSTOM_SIMPLE_AUTOMATION_TEST(FRemovingObjectRemovesFrequencyEntry, FChangeStreamsTestBase, "VirtualProduction.Concert.Replication.Stream.Frequency.RemovingObjectRemovesFrequencyEntry", EAutomationTestFlags::EditorContext | EAutomationTestFlags::EngineFilter);
 	bool FRemovingObjectRemovesFrequencyEntry::RunTest(const FString& Parameters)
 	{
 		// 1. Start up client with float property stream
@@ -226,7 +226,7 @@ namespace UE::ConcertSyncTests::Replication::Stream
 	 * - Cannot change frequency for unknown stream
 	 * - Cannot change frequency for object not registerd to stream
 	 */
-	IMPLEMENT_CUSTOM_SIMPLE_AUTOMATION_TEST(FPutReplacesAll, FChangeStreamsTestBase, "Editor.Concert.Replication.Stream.Frequency.PutReplacesAll", EAutomationTestFlags::EditorContext | EAutomationTestFlags::EngineFilter);
+	IMPLEMENT_CUSTOM_SIMPLE_AUTOMATION_TEST(FPutReplacesAll, FChangeStreamsTestBase, "VirtualProduction.Concert.Replication.Stream.Frequency.PutReplacesAll", EAutomationTestFlags::EditorContext | EAutomationTestFlags::EngineFilter);
 	bool FPutReplacesAll::RunTest(const FString& Parameters)
 	{
 		// 1. Start up client with float property stream
@@ -284,7 +284,7 @@ namespace UE::ConcertSyncTests::Replication::Stream
 	 * - Cannot change frequency for unknown stream
 	 * - Cannot change frequency for object not registerd to stream
 	 */
-	IMPLEMENT_CUSTOM_SIMPLE_AUTOMATION_TEST(FAddAfterPut, FChangeStreamsTestBase, "Editor.Concert.Replication.Stream.Frequency.AddAfterPut", EAutomationTestFlags::EditorContext | EAutomationTestFlags::EngineFilter);
+	IMPLEMENT_CUSTOM_SIMPLE_AUTOMATION_TEST(FAddAfterPut, FChangeStreamsTestBase, "VirtualProduction.Concert.Replication.Stream.Frequency.AddAfterPut", EAutomationTestFlags::EditorContext | EAutomationTestFlags::EngineFilter);
 	bool FAddAfterPut::RunTest(const FString& Parameters)
 	{
 		// 1. Start up client with float property stream
@@ -314,7 +314,7 @@ namespace UE::ConcertSyncTests::Replication::Stream
 	 * - Cannot change frequency for unknown stream
 	 * - Cannot change frequency for object not registerd to stream
 	 */
-	IMPLEMENT_CUSTOM_SIMPLE_AUTOMATION_TEST(FCannotChangeFrequencyOfUnregistered, FChangeStreamsTestBase, "Editor.Concert.Replication.Stream.Frequency.FailWithUnregisteredStreamOrObject", EAutomationTestFlags::EditorContext | EAutomationTestFlags::EngineFilter);
+	IMPLEMENT_CUSTOM_SIMPLE_AUTOMATION_TEST(FCannotChangeFrequencyOfUnregistered, FChangeStreamsTestBase, "VirtualProduction.Concert.Replication.Stream.Frequency.FailWithUnregisteredStreamOrObject", EAutomationTestFlags::EditorContext | EAutomationTestFlags::EngineFilter);
 	bool FCannotChangeFrequencyOfUnregistered::RunTest(const FString& Parameters)
 	{
 		// 1. Start up client with float property stream
diff --git a/Engine/Plugins/Developer/Concert/ConcertSync/ConcertSyncTest/Source/ConcertSyncTest/Private/Replication/Stream/StreamRequestTests_Query.cpp b/Engine/Plugins/Developer/Concert/ConcertSync/ConcertSyncTest/Source/ConcertSyncTest/Private/Replication/Stream/StreamRequestTests_Query.cpp
--- a/Engine/Plugins/Developer/Concert/ConcertSync/ConcertSyncTest/Source/ConcertSyncTest/Private/Replication/Stream/StreamRequestTests_Query.cpp
+++ b/Engine/Plugins/Developer/Concert/ConcertSync/ConcertSyncTest/Source/ConcertSyncTest/Private/Replication/Stream/StreamRequestTests_Query.cpp
@@ -22,7 +22,7 @@ namespace UE::ConcertSyncTests::Replication::Stream
 	 *
 	 * @note The other tests in StreamRequestTests_{x}.cpp use QueryStream under the hood so all this test does is checking that the skip flags work correctly. 
 	 */
-	IMPLEMENT_CUSTOM_SIMPLE_AUTOMATION_TEST(FQueryOtherClientStreams, FSendReceiveObjectTestBase, "Editor.Concert.Replication.Stream.QueryClientStream", EAutomationTestFlags::EditorContext | EAutomationTestFlags::EngineFilter);
+	IMPLEMENT_CUSTOM_SIMPLE_AUTOMATION_TEST(FQueryOtherClientStreams, FSendReceiveObjectTestBase, "VirtualProduction.Concert.Replication.Stream.QueryClientStream", EAutomationTestFlags::EditorContext | EAutomationTestFlags::EngineFilter);
 	bool FQueryOtherClientStreams::RunTest(const FString& Parameters)
 	{
 		// 1. Init
diff --git a/Engine/Plugins/Developer/Concert/ConcertSync/ConcertSyncTest/Source/ConcertSyncTest/Private/Replication/SyncControl.spec.cpp b/Engine/Plugins/Developer/Concert/ConcertSync/ConcertSyncTest/Source/ConcertSyncTest/Private/Replication/SyncControl.spec.cpp
--- a/Engine/Plugins/Developer/Concert/ConcertSync/ConcertSyncTest/Source/ConcertSyncTest/Private/Replication/SyncControl.spec.cpp
+++ b/Engine/Plugins/Developer/Concert/ConcertSync/ConcertSyncTest/Source/ConcertSyncTest/Private/Replication/SyncControl.spec.cpp
@@ -10,7 +10,7 @@
 
 namespace UE::ConcertSyncTests::Replication::UI
 {
-	BEGIN_DEFINE_SPEC(FReplicationSyncControlSpec, "Editor.Concert.Replication.SyncControl", EAutomationTestFlags::EditorContext | EAutomationTestFlags::EngineFilter)
+	BEGIN_DEFINE_SPEC(FReplicationSyncControlSpec, "VirtualProduction.Concert.Replication.SyncControl", EAutomationTestFlags::EditorContext | EAutomationTestFlags::EngineFilter)
 	
 	TSharedPtr<FReplicationServer> Server;
 	FReplicationClient* Sender = nullptr;
diff --git a/Engine/Plugins/Developer/Concert/ConcertSync/ConcertSyncTest/Source/ConcertSyncTest/Private/Replication/UI/GenericReplicationStreamModel.spec.cpp b/Engine/Plugins/Developer/Concert/ConcertSync/ConcertSyncTest/Source/ConcertSyncTest/Private/Replication/UI/GenericReplicationStreamModel.spec.cpp
--- a/Engine/Plugins/Developer/Concert/ConcertSync/ConcertSyncTest/Source/ConcertSyncTest/Private/Replication/UI/GenericReplicationStreamModel.spec.cpp
+++ b/Engine/Plugins/Developer/Concert/ConcertSync/ConcertSyncTest/Source/ConcertSyncTest/Private/Replication/UI/GenericReplicationStreamModel.spec.cpp
@@ -9,7 +9,7 @@
 
 namespace UE::ConcertSyncTests::Replication::UI
 {
-	BEGIN_DEFINE_SPEC(FReplicationStreamSpec, "Editor.Concert.Replication.UI", EAutomationTestFlags::EditorContext | EAutomationTestFlags::EngineFilter)
+	BEGIN_DEFINE_SPEC(FReplicationStreamSpec, "VirtualProduction.Concert.Replication.UI", EAutomationTestFlags::EditorContext | EAutomationTestFlags::EngineFilter)
 	FConcertObjectReplicationMap ReplicationMap;
 	TSharedPtr<ConcertSharedSlate::IEditableReplicationStreamModel> Model;
 	UTestReflectionObject* Object;
